[Mesh Partition]

Moved plugins into experimental.

#virtualized
#Synced-CL 51468759

diff --git a/Engine/Plugins/Experimental/MeshPartition/Shaders/MeshPartitionBorderFill.usf b/Engine/Plugins/Experimental/MeshPartition/Shaders/MeshPartitionBorderFill.usf
new file mode 100644
--- /dev/null
+++ b/Engine/Plugins/Experimental/MeshPartition/Shaders/MeshPartitionBorderFill.usf
@@ -0,0 +1,55 @@
+// Copyright Epic Games, Inc. All Rights Reserved.
+
+#include "/Engine/Private/Common.ush"
+
+uint2 Resolution;
+
+Texture2D<uint> Mask;
+
+RWTexture2D<float4> RWSectionTexture;
+
+[numthreads(8, 8, 1)]
+void MainCS(uint2 DTid : SV_DispatchThreadID)
+{
+	if (any(DTid >= Resolution) || Mask[DTid])
+	{
+		return;
+	}
+
+	// Naive filling, assume all mask candidates are valid
+	
+	const int BoxSearchExtent         = 4;
+	const int MaxCandidateDistanceSqr = Square(BoxSearchExtent);
+
+	const uint2 SearchMinOffset = max(int2(DTid) - BoxSearchExtent, 0);
+	const uint2 SearchMaxOffset = min(int2(DTid) + BoxSearchExtent, int2(Resolution - 1));
+
+	int    CandidateDistanceSqr = MaxCandidateDistanceSqr;
+	float4 CandidateTexel       = 0.0f;
+
+	for (uint CoordY = SearchMinOffset.y; CoordY <= SearchMaxOffset.y; CoordY++)
+	{
+		for (uint CoordX = SearchMinOffset.x; CoordX <= SearchMaxOffset.x; CoordX++)
+		{
+			int DistanceSqr = Square(int(CoordX) - int(DTid.x)) + Square(int(CoordY) - int(DTid.y));
+			if (DistanceSqr >= CandidateDistanceSqr)
+			{
+				continue;
+			}
+			
+			uint2 Offset = uint2(CoordX, CoordY);
+			if (!Mask[Offset])
+			{
+				continue;
+			}
+
+			CandidateDistanceSqr = DistanceSqr;
+			CandidateTexel       = RWSectionTexture[Offset];
+		}
+	}
+
+	if (CandidateDistanceSqr != MaxCandidateDistanceSqr)
+	{
+		RWSectionTexture[DTid] = CandidateTexel;
+	}
+}
diff --git a/Engine/Plugins/Experimental/MeshPartition/Shaders/MeshPartitionMakeSectionChannels.usf b/Engine/Plugins/Experimental/MeshPartition/Shaders/MeshPartitionMakeSectionChannels.usf
new file mode 100644
--- /dev/null
+++ b/Engine/Plugins/Experimental/MeshPartition/Shaders/MeshPartitionMakeSectionChannels.usf
@@ -0,0 +1,60 @@
+// Copyright Epic Games, Inc. All Rights Reserved.
+
+#include "/Engine/Private/Common.ush"
+
+// ----------------------------------------------------------------------------------
+
+Buffer<int>		InMeshIndices;	
+Buffer<float2>	InMeshUVs;
+Buffer<float>	InMeshChannelMasks;
+
+uint4	InDrawCall; // PageNum, NumChannelAttribs, FirstVertIndex
+uint4   InChannelsOffset; // Channel offsets from the very first channel attribute for this vertex
+float4  InChannelsMask;
+
+// ----------------------------------------------------------------------------------
+// Vertex shaders : 
+
+void MegaMeshDrawUVDomainVS(
+	uint VertexIndex : SV_VertexID,
+	out float2 SourceUV : TEXCOORD0,
+	out float4  OutMasks : TEXCOORD1,
+	out float4 OutPosition : SV_POSITION)
+{
+	uint VID = InMeshIndices[VertexIndex + InDrawCall.z];
+	
+	SourceUV = InMeshUVs[VID];
+	
+	// Channel Attrib Address is at VID x NumChannelAttribs
+	uint ChannelAttribOffset = VID * InDrawCall.y;
+	
+	// Fetch the 4 channels from the channel buffer (and maybe overflow fetching the neighbour vert channel values)
+	OutMasks = float4(	InMeshChannelMasks[ChannelAttribOffset + InChannelsOffset[0]],
+						InMeshChannelMasks[ChannelAttribOffset + InChannelsOffset[1]],
+						InMeshChannelMasks[ChannelAttribOffset + InChannelsOffset[2]],
+						InMeshChannelMasks[ChannelAttribOffset + InChannelsOffset[3]]);
+	
+	// Channels not written by this pass need to be reset to 0
+	OutMasks *= InChannelsMask;
+
+	float2 BaseUV = SourceUV;	
+	
+	BaseUV *= 2.0;
+	OutPosition.x = BaseUV.x - 1.0;
+	OutPosition.y = 1.0 - BaseUV.y;
+    OutPosition.zw = float2(0.0, 1.0);
+}
+
+// ----------------------------------------------------------------------------------
+// Pixel shaders : 
+
+void MegaMeshDrawUVDomainPS(
+	float2 SourceUV : TEXCOORD0,
+	float4 Masks : TEXCOORD1,
+	float4 SvPosition : SV_POSITION,
+	out float4 OutColor : SV_Target0,
+	out uint OutMask : SV_Target1)
+{
+	OutColor = float4(Masks);
+	OutMask  = 1;
+}
diff --git a/Engine/Plugins/Experimental/MeshPartition/Source/MeshPartition/Private/MeshPartitionStaticMeshComponent.cpp b/Engine/Plugins/Experimental/MeshPartition/Source/MeshPartition/Private/MeshPartitionStaticMeshComponent.cpp
new file mode 100644
--- /dev/null
+++ b/Engine/Plugins/Experimental/MeshPartition/Source/MeshPartition/Private/MeshPartitionStaticMeshComponent.cpp
@@ -0,0 +1,30 @@
+﻿// Copyright Epic Games, Inc. All Rights Reserved.
+
+#include "MeshPartitionStaticMeshComponent.h"
+
+int32 MaterialCacheTileCount = 256;
+static FAutoConsoleVariableRef CVarMegaMeshMaterialCacheTileCount(
+	TEXT("MegaMesh.MaterialCache.TileCount"),
+	MaterialCacheTileCount,
+	TEXT("The number of tiles on a given mega mesh section, effectively resolution driven by TileCount * r.MaterialCache.TileWidth")
+);
+
+
+namespace UE::MeshPartition
+{
+UMeshPartitionStaticMeshComponent::UMeshPartitionStaticMeshComponent(const FObjectInitializer& ObjectInitializer)
+{
+	MaterialCacheTileCount = FIntPoint(CVarMegaMeshMaterialCacheTileCount->GetInt(), CVarMegaMeshMaterialCacheTileCount->GetInt());
+	bLumenHeightfield = true;
+}
+
+void UMeshPartitionStaticMeshComponent::OnRegister()
+{
+	Super::OnRegister();
+}
+
+void UMeshPartitionStaticMeshComponent::OnUnregister()
+{
+	Super::OnUnregister();
+}
+} // namespace UE::MeshPartition
\ No newline at end of file
diff --git a/Engine/Plugins/Experimental/MeshPartition/Source/MeshPartition/Public/MeshPartitionComponent.h b/Engine/Plugins/Experimental/MeshPartition/Source/MeshPartition/Public/MeshPartitionComponent.h
new file mode 100644
--- /dev/null
+++ b/Engine/Plugins/Experimental/MeshPartition/Source/MeshPartition/Public/MeshPartitionComponent.h
@@ -0,0 +1,45 @@
+// Copyright Epic Games, Inc. All Rights Reserved.
+
+#pragma once
+
+#include "CoreMinimal.h"
+
+#include "Components/PrimitiveComponent.h"
+
+#include "MeshPartitionComponent.generated.h"
+
+#define UE_API MESHPARTITION_API
+
+namespace UE::MeshPartition
+{
+class UMeshPartitionDefinition;
+
+/**
+* Base class for a component of AMeshPartition actor.
+* This is a component to be able to be stripped during the cook process.
+* This should contain common interfaces between runtime and editor,
+* but also allows to implement the differences only useful for specific modes.
+*
+* #todo_megamesh [roey]: This is derived from UPrimitiveComponent so that it is targetable by any UPrimitiveComponentBackedTarget modeling tools.
+*                        Ideally this wouldn't be necessary, and perhaps we should solve this with a UMegaMeshTargetComponent specifically for this purpose
+*                        so that runtime is not affected.
+*/
+UCLASS(MinimalAPI, ClassGroup=(Custom), meta=(BlueprintSpawnableComponent))
+class UMeshPartitionComponent : public UPrimitiveComponent
+{
+	GENERATED_BODY()
+
+public:	
+	UE_API UMeshPartitionComponent();
+
+	virtual void PostRegisterMegaMeshComponents() {}
+	virtual void PostUnregisterMegaMeshComponents() {}
+
+	/**
+	* Callback called when the owner's currently used UMeshPartitionDefinition changed.
+	*/
+	virtual void OnDefinitionChanged(UMeshPartitionDefinition* InNewDefinition) {}
+};
+} // namespace UE::MeshPartition
+
+#undef UE_API
diff --git a/Engine/Plugins/Experimental/MeshPartition/Source/MeshPartitionCompute/Private/MeshPartitionChannelRasterizationShaders.cpp b/Engine/Plugins/Experimental/MeshPartition/Source/MeshPartitionCompute/Private/MeshPartitionChannelRasterizationShaders.cpp
new file mode 100644
--- /dev/null
+++ b/Engine/Plugins/Experimental/MeshPartition/Source/MeshPartitionCompute/Private/MeshPartitionChannelRasterizationShaders.cpp
@@ -0,0 +1,8 @@
+// Copyright Epic Games, Inc. All Rights Reserved.
+
+#include "MeshPartitionChannelRasterizationShaders.h"
+
+IMPLEMENT_GLOBAL_SHADER(FMegaMesh_DrawUVDomainVS, "/Plugin/MeshPartition/MeshPartitionMakeSectionChannels.usf", "MegaMeshDrawUVDomainVS", SF_Vertex);
+IMPLEMENT_GLOBAL_SHADER(FMegaMesh_DrawUVDomainPS, "/Plugin/MeshPartition/MeshPartitionMakeSectionChannels.usf", "MegaMeshDrawUVDomainPS", SF_Pixel);
+
+IMPLEMENT_GLOBAL_SHADER(FMegaMesh_BorderFillCS, "/Plugin/MeshPartition/MeshPartitionBorderFill.usf", "MainCS", SF_Compute);
diff --git a/Engine/Plugins/Experimental/MeshPartition/Source/MeshPartitionCompute/Public/MeshPartitionChannelRasterizationShaders.h b/Engine/Plugins/Experimental/MeshPartition/Source/MeshPartitionCompute/Public/MeshPartitionChannelRasterizationShaders.h
new file mode 100644
--- /dev/null
+++ b/Engine/Plugins/Experimental/MeshPartition/Source/MeshPartitionCompute/Public/MeshPartitionChannelRasterizationShaders.h
@@ -0,0 +1,63 @@
+// Copyright Epic Games, Inc. All Rights Reserved.
+
+#pragma once
+
+#include "GlobalShader.h"
+#include "PixelShaderUtils.h"
+#include "ShaderParameterStruct.h"
+#include "ShaderCompilerCore.h"
+
+#define UE_API MESHPARTITIONCOMPUTE_API
+
+//~ Plugins/MeshPartition/Shaders/MegaMeshMakeSectionChannels.usf shaders :
+class FMegaMesh_DrawUVDomainVS : public FGlobalShader
+{
+public:
+	DECLARE_EXPORTED_GLOBAL_SHADER(FMegaMesh_DrawUVDomainVS, UE_API);
+	SHADER_USE_PARAMETER_STRUCT(FMegaMesh_DrawUVDomainVS, FGlobalShader);
+
+	BEGIN_SHADER_PARAMETER_STRUCT(FParameters, )
+		SHADER_PARAMETER_RDG_BUFFER_SRV(Buffer<uint32>, InMeshIndices)
+		SHADER_PARAMETER_RDG_BUFFER_SRV(Buffer<float2>, InMeshUVs)
+		SHADER_PARAMETER_RDG_BUFFER_SRV(Buffer<float>, InMeshChannelMasks)
+		SHADER_PARAMETER(FUintVector4, InDrawCall)
+		SHADER_PARAMETER(FUintVector4, InChannelsOffset)
+		SHADER_PARAMETER(FVector4f, InChannelsMask)
+	END_SHADER_PARAMETER_STRUCT()
+};
+
+class FMegaMesh_DrawUVDomainPS : public FGlobalShader
+{
+public:
+	DECLARE_EXPORTED_GLOBAL_SHADER(FMegaMesh_DrawUVDomainPS, UE_API);
+	SHADER_USE_PARAMETER_STRUCT(FMegaMesh_DrawUVDomainPS, FGlobalShader);
+
+	BEGIN_SHADER_PARAMETER_STRUCT(FParameters, )
+		RENDER_TARGET_BINDING_SLOTS()
+	END_SHADER_PARAMETER_STRUCT()
+};
+
+class FMegaMesh_BorderFillCS : public FGlobalShader
+{
+public:
+	DECLARE_EXPORTED_GLOBAL_SHADER(FMegaMesh_BorderFillCS, UE_API);
+	SHADER_USE_PARAMETER_STRUCT(FMegaMesh_BorderFillCS, FGlobalShader);
+
+	BEGIN_SHADER_PARAMETER_STRUCT(FParameters, )
+		SHADER_PARAMETER(FUintVector2, Resolution)
+		SHADER_PARAMETER_RDG_TEXTURE_SRV(Texture2D<uint>,   Mask)
+		SHADER_PARAMETER_RDG_TEXTURE_UAV(RWTexture2D<float4>, RWSectionTexture)
+	END_SHADER_PARAMETER_STRUCT()
+
+	static void ModifyCompilationEnvironment(const FGlobalShaderPermutationParameters& InParameters, FShaderCompilerEnvironment& OutEnvironment)
+	{
+		OutEnvironment.CompilerFlags.Add(CFLAG_AllowTypedUAVLoads);
+	}
+};
+
+BEGIN_SHADER_PARAMETER_STRUCT(FMegaMesh_DrawUVDomain_Parameters, )
+	SHADER_PARAMETER_STRUCT_INCLUDE(FMegaMesh_DrawUVDomainVS::FParameters, VS)
+	SHADER_PARAMETER_STRUCT_INCLUDE(FMegaMesh_DrawUVDomainPS::FParameters, PS)
+END_SHADER_PARAMETER_STRUCT()
+
+#undef UE_API
diff --git a/Engine/Plugins/Experimental/MeshPartition/Source/MeshPartitionEditor/Private/MeshPartitionCollisionGeneration.h b/Engine/Plugins/Experimental/MeshPartition/Source/MeshPartitionEditor/Private/MeshPartitionCollisionGeneration.h
new file mode 100644
--- /dev/null
+++ b/Engine/Plugins/Experimental/MeshPartition/Source/MeshPartitionEditor/Private/MeshPartitionCollisionGeneration.h
@@ -0,0 +1,92 @@
+// Copyright Epic Games, Inc. All Rights Reserved.
+
+#pragma once
+
+#include "MeshPartitionChannel.h"
+#include "MeshPartitionCollisionComponent.h"
+#include "MeshPartitionDefinition.h"
+#include "MeshPartitionMeshData.h"
+
+#include "MeshPartitionCollisionGeneration.generated.h"
+
+// Support for generating collision from Mesh Partition meshes
+namespace UE::MeshPartition::Collision
+{
+	UENUM()
+	enum class EExperimentalCollisionSimplificationMethod
+	{
+		Cluster,
+		QEM
+	};
+
+	USTRUCT()
+	struct FExperimentalCollisionSimplificationSettings
+	{
+		GENERATED_BODY()
+
+	public:
+
+		// Whether to cut original mesh triangles along physical material boundaries. Otherwise, triangles will be assigned discretely to a material based on the average vertex weights.
+		// Note this introduces additional triangles initially, but can create smoother, more-accurate material boundaries that can be better-simplified.
+		UPROPERTY(EditAnywhere, Category = "Material Boundaries")
+		bool bCutMeshAlongMaterialBoundaries = true;
+
+		// When cutting material boundaries, tolerance distance at which to snap material boundary to existing geometry rather than cutting the mesh
+		UPROPERTY(EditAnywhere, Category = "Material Boundaries", meta = (EditCondition = bCutMeshAlongMaterialBoundaries))
+		float MaterialBoundaryCutSnapTolerance = 10.f;
+
+		UPROPERTY(EditAnywhere, Category = "Simplification",
+				  meta = (ToolTip = "Enable experimental collision simplification on the separate collision components"))
+		bool bSimplifyCollision = false;
+
+		UPROPERTY(EditAnywhere, Category = "Simplification")
+		EExperimentalCollisionSimplificationMethod SimplifyMethod = EExperimentalCollisionSimplificationMethod::QEM;
+
+		// Constrain how much simplification can move material boundaries away from their original positions
+		UPROPERTY(EditAnywhere, Category = "Simplification", meta = (EditCondition = "SimplifyMethod==EExperimentalCollisionSimplificationMethod::QEM"))
+		float MaterialBoundaryDistanceTolerance = 50.f;
+
+		// If positive, set a target edge length
+		UPROPERTY(EditAnywhere, Category = "Simplification")
+		float EdgeLength = -1;
+
+		// If positive, set a geometric error tolerance
+		UPROPERTY(EditAnywhere, Category = "Simplification", meta = (EditCondition = "SimplifyMethod==EExperimentalCollisionSimplificationMethod::QEM"))
+		float ErrorTolerance = 10;
+
+		// Scale up the requested accuracy (error metric and/or measured edge length) by up to this amount for upward-facing surfaces
+		UPROPERTY(EditAnywhere, Category = "Simplification")
+		float LocalAccuracyScale = 2.5;
+
+		// Min normal alignment angle vs Z-up, at which to finish scaling up simplifier accuracy, in degrees
+		UPROPERTY(EditAnywhere, Category = "Simplification", meta = (ClampMin = 0, ClampMax = 180, Units = Degrees))
+		float ScaleToMinNormalAngle = 26;
+
+		// Max normal alignment angle vs Z-up, at which to start scaling up simplifier accuracy, in degrees
+		UPROPERTY(EditAnywhere, Category = "Simplification", meta = (ClampMin=0, ClampMax=180, Units = Degrees))
+		float ScaleFromMaxNormalAngle = 53;
+
+		#if WITH_EDITOR
+		void MESHPARTITIONEDITOR_API GatherDependencies(MeshPartition::IDependencyInterface& Dependencies) const;
+		#endif
+	};
+
+	struct FMeshToCollisionSettings
+	{
+		UE::MeshPartition::Collision::FExperimentalCollisionSimplificationSettings SimplificationSettings;
+		
+		/** Physical material channels; vertex weight channels will be translated to per - triangle physical materials via this mapping */
+		TArray<UE::MeshPartition::FPhysicalMaterialChannel> PhysicalMaterialChannels;
+		
+		/** Default physical material, to set if channels not found */
+		UPhysicalMaterial* DefaultPhysicalMaterial = nullptr;
+
+		/** Prioritize cooking speed over runtime speed */
+		bool bFastCook = true;
+
+		/** Turn off ActiveEdgePrecompute (This makes cooking faster, but will slow contact generation) */
+		bool bDisableActiveEdgePrecompute = true;
+	};
+
+	void ConvertMeshToCollisionData(const FMeshData& Mesh, FMeshPartitionCollisionData& CollisionData, const FMeshToCollisionSettings& Settings);
+} // namespace UE::MeshPartition::Collision
diff --git a/Engine/Plugins/Experimental/MeshPartition/Source/MeshPartitionEditor/Private/MeshPartitionDescriptorCache.cpp b/Engine/Plugins/Experimental/MeshPartition/Source/MeshPartitionEditor/Private/MeshPartitionDescriptorCache.cpp
new file mode 100644
--- /dev/null
+++ b/Engine/Plugins/Experimental/MeshPartition/Source/MeshPartitionEditor/Private/MeshPartitionDescriptorCache.cpp
@@ -0,0 +1,364 @@
+// Copyright Epic Games, Inc. All Rights Reserved.
+
+#include "MeshPartitionDescriptorCache.h"
+
+#include "MeshPartition.h"
+#include "MeshPartitionEditorModule.h"			// LogMegaMeshEditor
+#include "MeshPartitionEditorUtils.h"			// PackageHashIsUpToDate
+#include "MeshPartitionDefinition.h"
+#include "MeshPartitionMeshBuilder.h"			// FindLayerPriorityIndexFromName
+#include "MeshPartitionCompiledSection.h"		// FCompiledSectionDescriptor
+#include "MeshPartitionModifierUtils.h"			// ForEachMegaMeshDescInWorldPartition, BetterHashCombine
+#include "MeshPartitionEditorSubsystem.h"		// FPIEPathFixer
+#include "WorldPartition/WorldPartition.h"
+#include "Algo/AnyOf.h"
+#include "Algo/ForEach.h"
+#include "EngineUtils.h"					// TActorIterator
+
+namespace UE::MeshPartition
+{
+FCachedDescriptors::FCachedDescriptors(const AMeshPartition* InMegaMesh)
+{
+	using namespace Utils;
+
+	check(InMegaMesh != nullptr);
+	check(InMegaMesh->GetActorGuid().IsValid());
+
+	UWorld* World = InMegaMesh->GetWorld();
+	check(World != nullptr);
+	check(!World->IsPlayInEditor());		// we shouldn't be caching descriptors for PIE worlds
+
+	UWorldPartition* WorldPartition = World->GetWorldPartition();
+	check(WorldPartition != nullptr)
+
+	const UMeshPartitionDefinition* MeshPartitionDefinition = InMegaMesh->GetMeshPartitionDefinition();
+	if (MeshPartitionDefinition == nullptr)
+	{
+		// UE_LOG(LogCore, Warning, TEXT("No MegaMesh Definition provided for %s (%.6s), using default settings."), *MegaMesh->GetName(), *MegaMesh->GetActorGuid().ToString());
+		MeshPartitionDefinition = UMeshPartitionDefinition::StaticClass()->GetDefaultObject<UMeshPartitionDefinition>();
+	}
+
+	// Grab ALL MegaMeshes present in the world, and find the descriptors for all the modifiers that belong to that megamesh, as well as the current compiled sections
+	// (NOTE: This only grabs from the main world - does not include anything from inside level instances. Those are gathered during the LoadActors() step)
+	TArray<FCompiledSectionDescriptor> SectionDescs;
+	ForEachMegaMeshDescInWorldPartition(WorldPartition,
+		[InMegaMesh, this](MeshPartition::FModifierDesc& Descriptor)
+		{
+			if (Descriptor.MegaMeshGuid == InMegaMesh->GetActorGuid())
+			{
+				AllModifiersGroup.Add(Descriptor);		// not yet sorted, may add in non-optimal order
+			}
+		},
+		[InMegaMesh, &SectionDescs](FCompiledSectionDescriptor& Descriptor)
+		{	if (Descriptor.Info.MegaMeshGUID == InMegaMesh->GetActorGuid())
+			{
+				SectionDescs.Add(Descriptor);
+			}
+		});
+
+	// Sort the Modifiers by priority:
+	AllModifiersGroup.Sort(MeshPartitionDefinition->GetModifierTypePriorities());
+
+	this->MegaMesh = InMegaMesh;
+	this->MegaMeshGuid = InMegaMesh->GetActorGuid();
+	this->CompiledSections = MoveTemp(SectionDescs);
+
+	// iterate all build variants
+	TConstArrayView<MeshPartition::FCompiledSectionBuildVariant> BuildVariantArray = MeshPartitionDefinition->GetCompiledSectionBuildVariants();
+	check(BuildVariantArray.Num() > 0);
+	for (const MeshPartition::FCompiledSectionBuildVariant& BuildVariant : BuildVariantArray)
+	{
+		// Break the MegaMesh into groups
+		this->BuildVariants.Add(BuildVariant.Name, MeshPartition::FModifierGrouping(BuildVariant, *MeshPartitionDefinition, AllModifiersGroup));
+	}
+}
+
+FCachedDescriptors::FCachedDescriptors(const AMeshPartition* InMegaMesh, TConstArrayView<MeshPartition::FModifierDesc> InAllModifiersSorted, TArray<FCompiledSectionDescriptor>&& InCompiledSections)
+{
+	check(InMegaMesh != nullptr);
+	this->MegaMeshGuid = InMegaMesh->GetActorGuid();
+	check(this->MegaMeshGuid.IsValid());
+
+	UWorld* World = InMegaMesh->GetWorld();
+	check(World != nullptr);
+
+	UWorldPartition* WorldPartition = World->GetWorldPartition();
+	check(WorldPartition != nullptr)
+
+	const UMeshPartitionDefinition* MeshPartitionDefinition = InMegaMesh->GetMeshPartitionDefinition();
+	if (MeshPartitionDefinition == nullptr)
+	{
+		UE_LOG(LogCore, Warning, TEXT("No MegaMesh Definition provided for %s (%.6s), using default settings."), *InMegaMesh->GetName(), *InMegaMesh->GetActorGuid().ToString());
+		MeshPartitionDefinition = UMeshPartitionDefinition::StaticClass()->GetDefaultObject<UMeshPartitionDefinition>();
+	}
+
+	// copy the sorted modifier list into the group
+	for (const MeshPartition::FModifierDesc& ModifierDesc : InAllModifiersSorted)
+	{
+		AllModifiersGroup.Add(ModifierDesc);
+	}
+
+	// Sort the Modifiers by priority:
+	AllModifiersGroup.Sort(MeshPartitionDefinition->GetModifierTypePriorities());
+
+	// Break the MegaMesh into groups
+	this->MegaMesh = InMegaMesh;
+	this->MegaMeshGuid = InMegaMesh->GetActorGuid();
+	this->CompiledSections = MoveTemp(InCompiledSections);
+
+	// iterate all build variants
+	TConstArrayView<MeshPartition::FCompiledSectionBuildVariant> BuildVariantArray = MeshPartitionDefinition->GetCompiledSectionBuildVariants();
+	check(BuildVariantArray.Num() > 0);
+	for (const MeshPartition::FCompiledSectionBuildVariant& BuildVariant : BuildVariantArray)
+	{
+		// Break the MegaMesh into groups
+		this->BuildVariants.Add(BuildVariant.Name, MeshPartition::FModifierGrouping(BuildVariant, *MeshPartitionDefinition, AllModifiersGroup));
+	}
+}
+
+FModifierGrouping::FModifierGrouping(const MeshPartition::FCompiledSectionBuildVariant& InBuildVariant, const UMeshPartitionDefinition& Definition, const MeshPartition::FModifierGroup& AllModifiersGroupSorted)
+{
+	BuildVariantName = InBuildVariant.Name;
+
+	MeshPartition::FBuilderSettings BuilderSettings;
+	// these are the only two properties that affect grouping
+	BuilderSettings.TypePriorities = Definition.GetModifierTypePriorities();
+	BuilderSettings.MaxSectionComplexity = InBuildVariant.MaxSectionComplexity;
+
+	ModifierGroups = BuildModifierGroups(AllModifiersGroupSorted.AllModifierDescs(), BuilderSettings);
+}
+
+MeshPartition::FModifierGroup FCachedDescriptors::BuildModifierGroupForBaseModifierPaths(TConstArrayView<FSoftObjectPath> BaseModifierPaths) const
+{
+	// convert BaseModifierPaths into Descriptors
+	TArray<MeshPartition::FModifierDesc> BaseSet;
+	{
+		for (FSoftObjectPath Path : BaseModifierPaths)
+		{
+			// NOTE that the base modifier paths passed in have to correspond to serialized paths (they cannot be temporary PIE paths)
+			const MeshPartition::FModifierDesc* Desc = AllModifiersGroup.AllModifierDescs().FindByPredicate([&Path](const MeshPartition::FModifierDesc& Desc) { return Desc.ModifierPath == Path; });
+			check(Desc);
+			BaseSet.Add(*Desc);
+		}
+	}
+
+	return BuildModifierGroupForBase(BaseSet, AllModifiersGroup.ModifierDescs());
+}
+
+void FCachedDescriptors::ForAllBuildVariants(TFunctionRef<bool(const MeshPartition::FModifierGrouping&)> InFunc) const
+{
+	using FVariantNameCachedBuildVariantPair = TPair<FName, MeshPartition::FModifierGrouping>;
+
+	TArray<FName> Keys;
+
+    Keys.Reserve(BuildVariants.Num());
+
+	for (const FVariantNameCachedBuildVariantPair& VariantNameCachedBuildVariantPair : BuildVariants)
+	{
+		Keys.Emplace(VariantNameCachedBuildVariantPair.Key);
+	}
+
+	for (const FName& Key : Keys)
+	{
+		const MeshPartition::FModifierGrouping* BuildVariant = BuildVariants.Find(Key);
+
+		if (BuildVariant == nullptr)
+		{
+			continue;
+		}
+
+		if (!InFunc(*BuildVariant))
+		{
+			return;
+		}
+	}
+}
+
+void FModifierDescriptorCache::AddMegaMeshToGuidMap(const AMeshPartition* MegaMesh)
+{
+	if (MegaMesh == nullptr)
+	{
+		return;
+	}
+
+	const FGuid& MegaMeshGuid = MegaMesh->GetActorGuid();
+	if (MegaMeshGuid.IsValid())
+	{
+		FSoftObjectPath MegaMeshPath(MegaMesh);
+		GuidToMegaMesh.AddUnique(MegaMeshGuid, MegaMeshPath);
+	}
+}
+
+void FModifierDescriptorCache::UpdateCacheForAllMegaMeshesInWorld(UWorld* InWorld, TArray<FGuid> *OutAllMegaMeshGuidsInWorld)
+{
+	using namespace Utils;
+
+	TMap<FGuid, TTuple<TArray<MeshPartition::FModifierDesc>, TArray<FCompiledSectionDescriptor>>> MegaMeshesInWorld;
+	UWorldPartition* WorldPartition = InWorld->GetWorldPartition();
+	if (WorldPartition == nullptr)
+	{
+		return;
+	}
+
+	UE::MeshPartition::EditorUtils::FPIEPathFixer& PIEPathFixer = UMeshPartitionEditorSubsystem::GetPIEPathFixer();
+	int32 CompiledSections = 0;
+	int32 Modifiers = 0;
+	int32 SectionGroups = 0;
+
+	// check if this is a PIE world
+	bool bIsPIE = InWorld->IsPlayInEditor();
+
+	// NOTE: this only grabs actor descs from the Main world/level -- not from any level instances
+	ForEachMegaMeshDescInWorldPartition(WorldPartition,
+		[&MegaMeshesInWorld, &Modifiers, &PIEPathFixer, bIsPIE](MeshPartition::FModifierDesc& Descriptor)
+		{
+			if (Descriptor.MegaMeshGuid.IsValid())
+			{
+				TArray<MeshPartition::FModifierDesc>& ModifierDescriptors = MegaMeshesInWorld.FindOrAdd(Descriptor.MegaMeshGuid).Get<0>();
+
+				// TODO: this may not correctly modify a level instanced path.. we may need to strip PIE paths from those individually?
+				FSoftObjectPath ModifierPath = Descriptor.ModifierPath;
+				if (bIsPIE)
+				{
+					// replace the PIE package name with the non-PIE package name
+					PIEPathFixer.FixInPlace(Descriptor.ModifierPath);
+				}
+
+				ModifierDescriptors.Add(Descriptor);
+				Modifiers++;
+			}
+		},
+		[&MegaMeshesInWorld, &CompiledSections](FCompiledSectionDescriptor& Descriptor)
+		{
+			// ignore old descriptors that don't have megamesh GUID populated
+			if (Descriptor.Info.MegaMeshGUID.IsValid())
+			{
+				TArray<FCompiledSectionDescriptor>& CompiledSectionDescriptors = MegaMeshesInWorld.FindOrAdd(Descriptor.Info.MegaMeshGUID).Get<1>();
+				CompiledSectionDescriptors.Add(Descriptor);
+				CompiledSections++;
+			}
+		});
+
+	// Build the section groups for each megamesh in the world
+	// NOTE: this will not gather spatially streaming AMegaMesh actors
+	// (currently only the section actors should be streaming, not the AMegaMesh actor)
+	for (TActorIterator<AMeshPartition> It(InWorld); It; ++It)
+	{
+		const AMeshPartition* MegaMesh = *It;
+		const FGuid MegaMeshGuid = MegaMesh->GetActorGuid();
+
+		// populate the GUID -> MegaMesh with the MegaMeshes that exist in the world
+		AddMegaMeshToGuidMap(MegaMesh);
+
+		TTuple<TArray<MeshPartition::FModifierDesc>, TArray<FCompiledSectionDescriptor>>* Descriptors = MegaMeshesInWorld.Find(MegaMeshGuid);
+		if ((Descriptors == nullptr) || (Descriptors->Get<0>().Num() <= 0))
+		{
+			// no descriptors found.. this megamesh will be empty
+			UE_LOG(LogMegaMeshEditor, Warning, TEXT("MegaMesh '%s' (GUID '%s') does not have any modifiers in the world, the megamesh will be empty"),
+				*FSoftObjectPath(MegaMesh).ToString(),
+				*MegaMeshGuid.ToString());
+			continue;
+		}
+
+		// and try to build the section groups for it
+		FSoftObjectPath ObjectPath(MegaMesh);
+		TSharedPtr<MeshPartition::FCachedDescriptors>* CacheEntryPtr = CachedDescriptors.Find(ObjectPath);
+		TSharedPtr<MeshPartition::FCachedDescriptors> CacheEntry;
+		TConstArrayView<MeshPartition::FModifierDesc> AllModifiersSorted = Descriptors->Get<0>();
+		TArray<FCompiledSectionDescriptor>& AllCompiledSections = Descriptors->Get<1>();
+		if (CacheEntryPtr == nullptr)
+		{
+			CacheEntry = MakeShared<MeshPartition::FCachedDescriptors>(MegaMesh, AllModifiersSorted, MoveTemp(AllCompiledSections));
+			CachedDescriptors.Emplace(ObjectPath, CacheEntry);
+		}
+		else
+		{
+			// TODO: we could validate and/or update the existing section...  for now we just stomp it
+			CacheEntry = MakeShared<MeshPartition::FCachedDescriptors>(MegaMesh, AllModifiersSorted, MoveTemp(AllCompiledSections));
+			CachedDescriptors.Emplace(ObjectPath, CacheEntry);
+		}
+	}
+	
+	UE_LOG(LogCore, Verbose, TEXT("UpdateCacheForAllMegaMeshesInWorld (%d Modifiers, %d Compiled Sections)"), Modifiers, CompiledSections);
+
+	if (OutAllMegaMeshGuidsInWorld != nullptr)
+	{
+		MegaMeshesInWorld.GetKeys(*OutAllMegaMeshGuidsInWorld);
+	}
+}
+
+TSharedPtr<MeshPartition::FCachedDescriptors> FModifierDescriptorCache::GetCachedDescriptorsByPath(FSoftObjectPath MegaMeshPath)
+{
+	if (TSharedPtr<MeshPartition::FCachedDescriptors>* CacheEntryPtr = CachedDescriptors.Find(MegaMeshPath))
+	{
+		return *CacheEntryPtr;
+	}
+	return nullptr;
+}
+
+TSharedPtr<MeshPartition::FCachedDescriptors> FModifierDescriptorCache::GetCachedDescriptorsByGuid(FGuid MegaMeshGUID)
+{
+	TSharedPtr<MeshPartition::FCachedDescriptors>* CacheEntryPtr = nullptr;
+
+	TArray<FSoftObjectPath> MegaMeshPaths;
+	GuidToMegaMesh.MultiFind(MegaMeshGUID, MegaMeshPaths);
+
+	for (FSoftObjectPath MegaMeshPath : MegaMeshPaths)
+	{
+		CacheEntryPtr = CachedDescriptors.Find(MegaMeshPath);
+		if (CacheEntryPtr != nullptr)
+		{
+			return *CacheEntryPtr;
+		}
+	}
+
+	return nullptr;
+}
+
+void FModifierDescriptorCache::ForAllCachedDescriptors(TFunctionRef<bool(const FSoftObjectPath&, const TSharedPtr<MeshPartition::FCachedDescriptors>&)> InFunc) const
+{
+	using FMeshPartitionDescriptorsPair = TPair<FSoftObjectPath, TSharedPtr<MeshPartition::FCachedDescriptors>>;
+
+	TArray<FSoftObjectPath> Keys;
+
+    Keys.Reserve(CachedDescriptors.Num());
+
+	for (const FMeshPartitionDescriptorsPair& MeshPartitionDescriptorsPair : CachedDescriptors)
+	{
+		Keys.Emplace(MeshPartitionDescriptorsPair.Key);
+	}
+
+	for (const FSoftObjectPath& Key : Keys)
+	{
+		const TSharedPtr<MeshPartition::FCachedDescriptors>* CachedDescriptor = CachedDescriptors.Find(Key);
+
+		if (CachedDescriptor == nullptr)
+		{
+			continue;
+		}
+
+		if (!InFunc(Key, *CachedDescriptor))
+		{
+			return;
+		}
+	}
+}
+
+TSharedPtr<MeshPartition::FCachedDescriptors> FModifierDescriptorCache::GetCachedDescriptors(
+	const AMeshPartition* MegaMesh)
+{
+	check(MegaMesh);
+	FSoftObjectPath MegaMeshPath(MegaMesh);
+	
+	if (TSharedPtr<MeshPartition::FCachedDescriptors>* CacheEntryPtr = CachedDescriptors.Find(MegaMeshPath))
+	{
+		return *CacheEntryPtr;
+	}
+	else
+	{
+		// try looking it up by GUID... this will effectively grab the cached descriptors from the editor-world
+		FGuid MegaMeshGuid = MegaMesh->GetActorGuid();
+		return GetCachedDescriptorsByGuid(MegaMeshGuid);
+	}
+}
+} // namespace UE::MeshPartition
diff --git a/Engine/Plugins/Experimental/MeshPartition/Source/MeshPartitionEditor/Private/MeshPartitionModiiferComponentDesc.cpp b/Engine/Plugins/Experimental/MeshPartition/Source/MeshPartitionEditor/Private/MeshPartitionModiiferComponentDesc.cpp
new file mode 100644
--- /dev/null
+++ b/Engine/Plugins/Experimental/MeshPartition/Source/MeshPartitionEditor/Private/MeshPartitionModiiferComponentDesc.cpp
@@ -0,0 +1,58 @@
+// Copyright Epic Games, Inc. All Rights Reserved.
+
+#include "MeshPartitionModifierComponentDesc.h"
+#include "MeshPartitionModifierComponent.h"
+#include "WorldPartition/WorldPartitionActorDescArchive.h"
+
+namespace UE::MeshPartition
+{
+
+uint32 FWorldPartitionModifierComponentDesc::GetSizeOf() const
+{
+	return sizeof(*this);
+}
+
+void FWorldPartitionModifierComponentDesc::Init(const UActorComponent* InComponent)
+{
+	FWorldPartitionComponentDesc::Init(InComponent);
+
+	const MeshPartition::UModifierComponent* Modifier = CastChecked<MeshPartition::UModifierComponent>(InComponent);
+	MeshPartition::FModifierDesc Desc(*Modifier);
+
+	// We must store the bounds of the component in local space so that delta serialization based on the default of the component
+	// can properly function. If not, the bounds will always differ from the default and will not correctly respond to changes in
+	// the base class which affect the modifier bounds.
+	ComponentTransform = Modifier->GetComponentTransform();
+	LocalBounds = Desc.Bounds.InverseTransformBy(ComponentTransform);
+	
+	ModifierPath = Desc.ModifierPath;
+	Type = Desc.Type;
+	Priority = Desc.Priority;
+	MegaMeshGuid = Desc.MegaMeshGuid;
+	BaseGrowth = Desc.BaseGrowth;
+	Complexity = Desc.Complexity;
+	ComplexityMultiplier = Desc.ComplexityMultiplier;
+	bIsContiguous = Desc.bIsContiguous;
+	bIsDisabled = Desc.bIsDisabled;
+
+}
+
+void FWorldPartitionModifierComponentDesc::Serialize(FArchive& Ar)
+{
+	FWorldPartitionComponentDesc::Serialize(Ar);
+
+	Ar << ComponentTransform;
+	Ar << ModifierPath; // modifier path is technically serializable from the parent actor desc + component desc but there are edge cases and this is far simpler.
+	Ar << MegaMeshGuid;
+
+	Ar << TDeltaSerialize<FName>(Type);
+	Ar << TDeltaSerialize<double>(Priority);
+	Ar << TDeltaSerialize<FBox>(LocalBounds);
+	Ar << TDeltaSerialize<MeshPartition::FBaseGrowth>(BaseGrowth);
+	Ar << TDeltaSerialize<double>(Complexity);
+	Ar << TDeltaSerialize<float>(ComplexityMultiplier);
+	Ar << TDeltaSerialize<bool>(bIsContiguous);
+	Ar << TDeltaSerialize<bool>(bIsDisabled);
+}
+
+} // namespace UE::MeshPartition
diff --git a/Engine/Plugins/Experimental/MeshPartition/Source/MeshPartitionEditor/Private/MeshPartitionPlatformCellTransformer.cpp b/Engine/Plugins/Experimental/MeshPartition/Source/MeshPartitionEditor/Private/MeshPartitionPlatformCellTransformer.cpp
new file mode 100644
--- /dev/null
+++ b/Engine/Plugins/Experimental/MeshPartition/Source/MeshPartitionEditor/Private/MeshPartitionPlatformCellTransformer.cpp
@@ -0,0 +1,143 @@
+// Copyright Epic Games, Inc. All Rights Reserved.
+
+#include "MeshPartitionPlatformCellTransformer.h"
+#include "Interfaces/ITargetPlatformManagerModule.h"
+#include "Interfaces/ITargetPlatform.h"
+
+#include "Engine/Level.h"
+#include "GameFramework/Actor.h"
+#include "Editor/EditorEngine.h"
+#include "Algo/Transform.h"
+#include "Algo/AnyOf.h"
+#include "WorldPartition/DataLayer/DataLayerInstance.h"
+#include "WorldPartition/DataLayer/DataLayerManager.h"
+#include "WorldPartition/WorldPartitionRuntimeCell.h"
+#include "MeshPartitionDataLayerContainer.h"
+#include "MeshPartitionEditorModule.h"
+#include "MeshPartitionDefinition.h"
+#include "WorldPartition/HLOD/HLODActor.h"
+#include "MeshPartitionEditorUtils.h"
+
+#include "MeshPartitionCompiledSection.h"
+
+extern UNREALED_API UEditorEngine* GEditor;
+
+namespace UE::MeshPartition
+{
+
+void UPlatformCellTransformer::TransformRuntimeCell(UWorldPartitionRuntimeCell* InCell)
+{
+	// Whenever a build variant set of CompiledSections are generated, they are assigned a unique data layer each.
+	// A "build variant" is a set of compiled sections which can be enabled or disabled for specific target platforms.
+	// sections per build variant may also have completely different bounds and generate completely different numbers of compiled section actors.
+	// (you might want compiled sections with nanite enabled for PC and a set with nanite disabled and traditional LODs for switch).
+	// (for servers we have a "Common" build variant which contains only the runtime gameplay data and rendering is split off into their own variants)
+	//
+	// This data layer is used to allow us to separate the compiled sections (and derived HLODs) into their own streaming cells. This prevents all compiled
+	// sections across build variants being included in a single HLOD and it prevents the larger high end platform actor sizes from polluting the streaming
+	// grids of lower end platforms.
+	//
+	// During Transform(UWPRuntimeCell*) we need to strip these data layer so that we can achieve normal data layer loading properties for compiled sections.
+	// This is necessary to solve the following cases:
+	// 1. User assigned a "classic" runtime data layer to their compiled section actors and triggered it to load.
+	//		this would enable loading of _all_ mesh partition build variants due to the OR nature of DLs.
+	// 2. User assigned a "classic" runtime data layer to their compiled section actors and expects it to remain unloaded (Eg. AncientGame nighttime scene)
+	//		When the mesh partition system tries to load the active platform data layer it would unintentionally load data which was intended to be unloaded.
+
+	TArray<const UDataLayerInstance*> CellDataLayers = InCell->GetDataLayerInstances();
+
+	TArray<TObjectPtr<const UDataLayerAsset>> StrippedDataLayers;
+	for (auto It = CellDataLayers.CreateIterator(); It; ++It)
+	{
+		const UDataLayerInstance* Instance = *It;
+		if (const UDataLayerAsset* Asset = Instance->GetAsset())
+		{
+			if (const MeshPartition::AMeshPartitionDataLayerContainer* Container = Asset->GetTypedOuter<MeshPartition::AMeshPartitionDataLayerContainer>())
+			{
+				if (Container->IsDataLayerOwnedByContainer(Asset))
+				{
+					StrippedDataLayers.Add(Asset);
+					It.RemoveCurrentSwap();
+				}
+			}
+		}
+	}
+
+	if (StrippedDataLayers.Num())
+	{
+		MeshPartition::FStrippedDataLayerList& CellStrippedDataLayers = CellToStrippedDataLayers.FindOrAdd(InCell->GetGuid());
+		CellStrippedDataLayers.List = MoveTemp(StrippedDataLayers);
+	}
+	
+	InCell->SetDataLayers(CellDataLayers);
+}
+
+void UPlatformCellTransformer::Transform(ULevel* InLevel)
+{
+	// During Transform(ULevel*) we have to strip all the actors which are not relevant for the current cook platform
+	// so they do not end up in the final packaged data.
+
+	// #TODO [roey]: How to handle cooks with -TargetPlatform=WindowsServer+Linux+MacOS?
+	// For now it will probably be fine if those actors just end up in the game and are not loaded due to
+	// AActor::NeedsLoadForTargetPlatform -> false.
+	// Since the streaming grids are split by the data layers already it isn't such a big deal as they will not
+	// be able to affect each other's streaming cells. The main downside is package size. There don't seem to be
+	// any cooks with multiple target platforms in the same process so maybe this is fine for now.
+	// Ideally this stripping per-platform would be handled by the WP system internally but it's not clear how that can be
+	// done with the current cook architecture. I believe that each invocation of the cook commandlet will only generate a
+	// single set of package via the WP package splitter for each world.
+	
+	const TArray<ITargetPlatform*> TargetPlatforms = UE::MeshPartition::EditorUtils::GetTargetPlatforms();
+
+	const UWorldPartitionRuntimeCell* Cell = CastChecked<UWorldPartitionRuntimeCell>(InLevel->GetWorldPartitionRuntimeCell());
+
+	const MeshPartition::AMeshPartitionDataLayerContainer* Container = GetDataLayerContainerForCell(Cell);
+
+	// if the map has no data layer container there is nothing for us to do.
+	if (!Container)
+	{
+		return;
+	}
+
+	// If the current cell has any MeshPartition Data Layers which were stripped out we need to inspect the contents
+	// and determine if we should wipe this cell or retain it.
+	if (const MeshPartition::FStrippedDataLayerList* DataLayers = CellToStrippedDataLayers.Find(Cell->GetGuid()))
+	{
+		bool bKeepCellContents = false;
+		for (TObjectPtr<const UDataLayerAsset> DataLayer : DataLayers->List)
+		{
+			if (bKeepCellContents)
+			{
+				break;
+			}
+
+			for (ITargetPlatform* TargetPlatform : TargetPlatforms)
+			{
+				if (Container->IsDataLayerRelevantForPlatform(DataLayer, TargetPlatform))
+				{
+					bKeepCellContents = true;
+					break;
+				}
+			}
+		}
+
+		if (!bKeepCellContents)
+		{
+			// if the cell contained any removed compiled sections it implies that this cell is for a build config which should not be loaded
+			// for the target platform. In this case, we must clear the entire cell contents.
+			// This keeps derived actors such as HLODs from the cells.
+			// There is also no other valid way for non-compiled section actors to enter the cell.
+			UE_LOG(LogMegaMeshEditor, Verbose, TEXT("Stripped all actos from WPRuntimeCell (%s) due to cell not being relevant for target platform."), *InLevel->GetWorldPartitionRuntimeCell()->GetDebugName());
+			InLevel->Actors.Empty();
+		}
+	}
+}
+
+const MeshPartition::AMeshPartitionDataLayerContainer* UPlatformCellTransformer::GetDataLayerContainerForCell(const UWorldPartitionRuntimeCell* InCell)
+{
+	// Retrieve the data layer container from the **cell**'s outer world. This is the main partition. Trying to find it in
+	// `InLevel->GetWorld()` will fail since the streaming level is not yet linked to the main world.
+	return AMeshPartitionDataLayerContainer::Get(InCell->GetOuterWorld());
+}
+
+}
diff --git a/Engine/Plugins/Experimental/MeshPartition/Source/MeshPartitionEditor/Private/MeshPartitionPreviewComponents.cpp b/Engine/Plugins/Experimental/MeshPartition/Source/MeshPartitionEditor/Private/MeshPartitionPreviewComponents.cpp
new file mode 100644
--- /dev/null
+++ b/Engine/Plugins/Experimental/MeshPartition/Source/MeshPartitionEditor/Private/MeshPartitionPreviewComponents.cpp
@@ -0,0 +1,657 @@
+// Copyright Epic Games, Inc. All Rights Reserved.
+
+#include "MeshPartitionPreviewComponents.h"
+
+#include "PrimitiveSceneProxy.h"
+#include "DynamicMeshBuilder.h"
+#include "SceneInterface.h"
+#include "Materials/Material.h"
+#include "MaterialDomain.h"
+#include "SceneView.h"
+#include "Engine/Engine.h"
+#include "MaterialShared.h"
+#include "Materials/MaterialRenderProxy.h"
+#include "PrimitiveUniformShaderParametersBuilder.h"
+#include "MaterialCache/MaterialCache.h"
+#include "MaterialCache/MaterialCacheVirtualTexture.h"
+#include "ProfilingDebugging/ScopedTimers.h"
+#include "MeshPartitionEditorModule.h"
+
+#if RHI_RAYTRACING
+#include "RayTracingInstance.h"
+#endif // RHI_RAYTRACING
+
+#include "MaterialCachedData.h"
+#include "MeshPartitionMeshData.h"
+
+namespace UE::MeshPartition
+{
+static TAutoConsoleVariable<bool> CVarCustomMeshGeneratePhysics(TEXT("MegaMesh.Preview.CustomPreviewMeshGeneratePhysics"),
+	true,
+	TEXT("Generate physics data for Mesh Terrain preview meshes, to aid editor interaction."));
+
+static TAutoConsoleVariable<bool> CVarCustomMeshAsyncPhysics(TEXT("MegaMesh.Preview.CustomPreviewMeshAsyncPhysics"),
+	true, 
+	TEXT("Allow async generation of physics data for Mesh Terrain preview meshes."));
+
+HHitProxy* UStaticMeshPreviewComponent::CreateMeshHitProxy(int32 SectionIndex, int32 MaterialIndex) const
+{
+	return new MeshPartition::HPreviewProxy(this);
+}
+
+UPreviewMeshComponent::UPreviewMeshComponent()
+{
+	MeshData = MakeShared<const FMeshData>();
+}
+
+void UPreviewMeshComponent::GetResourceSizeEx(FResourceSizeEx& Size)
+{
+	Super::GetResourceSizeEx(Size);
+
+	Size.AddDedicatedSystemMemoryBytes(MeshData ? MeshData->GetByteCount() : 0);
+}
+
+void UPreviewMeshComponent::SetMeshData(TSharedRef<const MeshPartition::FMeshData> InMeshData)
+{
+	MeshData = InMeshData;
+	
+	MarkRenderStateDirty();
+	UpdateBounds();
+	{
+		double Time = 0.0;
+		FDurationTimer Timer(Time);
+		Timer.Start();
+
+		RebuildPhysicsData();
+
+		Timer.Stop();
+		UE_LOG(LogMegaMeshEditor, Verbose, TEXT("UPreviewMeshComponent::SetMeshData created physics state in %.2f seconds"), Time);
+	}
+}
+
+void UPreviewMeshComponent::SetMeshData(FMeshData&& InMeshData)
+{
+	SetMeshData(MakeShared<const FMeshData>(MoveTemp(InMeshData)));
+}
+
+FBoxSphereBounds UPreviewMeshComponent::CalcBounds(const FTransform& LocalToWorld) const
+{
+	FBox Result;
+
+	for (int VertexID : MeshData->VertexIndicesItr())
+	{
+		Result += MeshData->GetVertex(VertexID);
+	}
+	if (!Result.IsValid)
+	{
+		// If bbox is empty, set a very small bbox to avoid log spam/etc in other engine systems.
+		// The check used is generally IsNearlyZero(), which defaults to KINDA_SMALL_NUMBER, so set 
+		// a slightly larger box here to be above that threshold
+		constexpr double Extent = (double)(KINDA_SMALL_NUMBER + SMALL_NUMBER);
+		Result = FBox(FVector(-Extent), FVector(Extent));
+	}
+	return FBoxSphereBounds(Result).TransformBy(LocalToWorld);
+}
+
+void UPreviewMeshComponent::OnRegister()
+{
+	Super::OnRegister();
+}
+
+void UPreviewMeshComponent::OnUnregister()
+{
+	Super::OnUnregister();
+
+	for (TObjectPtr VirtualTexture : MaterialCacheTextures)
+	{
+		VirtualTexture->Unregister();
+	}
+
+	MaterialCacheTextures.Empty();
+
+	AbortPendingPhysicsCook();
+}
+
+void UPreviewMeshComponent::RebuildPhysicsData()
+{
+	if (AsyncBodySetup)
+	{
+		AsyncBodySetup->AbortPhysicsMeshAsyncCreation();
+		AsyncBodySetup = nullptr;
+	}
+
+	if (!CVarCustomMeshAsyncPhysics.GetValueOnGameThread() || !CVarCustomMeshGeneratePhysics.GetValueOnGameThread() ||
+		!MeshData.IsValid() || MeshData->VertexCount() <= 0)
+	{
+		MeshBodySetup = nullptr;
+		// Generate synchronously
+		RecreatePhysicsState();
+		return;
+	}
+
+	AsyncBodySetup = CreateBodySetupHelper();
+	UE_LOG(LogMegaMeshEditor, Verbose, TEXT("UPreviewMeshComponent::RebuildPhysicsData enqueuing async physics generation.  guid=%s"), *AsyncBodySetup->BodySetupGuid.ToString());
+
+	AsyncBodySetup->CreatePhysicsMeshesAsync(FOnAsyncPhysicsCookFinished::CreateUObject(this, &UPreviewMeshComponent::FinishPhysicsAsyncCook, AsyncBodySetup.Get()));
+}
+
+void UPreviewMeshComponent::FinishPhysicsAsyncCook(bool bSuccess, UBodySetup* FinishedBodySetup)
+{
+	if (FinishedBodySetup != AsyncBodySetup)
+	{
+		UE_LOG(LogMegaMeshEditor, Display, TEXT("UPreviewMeshComponent::FinishPhysicsAsyncCook discarding cancelled UBodySetup generation.  guid=%s"), *FinishedBodySetup->BodySetupGuid.ToString());
+		return;
+	}
+
+	// Received the UBodySetup containing the cooked physics data.  Apply it.
+	UE_LOG(LogMegaMeshEditor, Verbose, TEXT("UPreviewMeshComponent::FinishPhysicsAsyncCook applying new UBodySetup.  guid=%s"), *FinishedBodySetup->BodySetupGuid.ToString());
+	MeshBodySetup = FinishedBodySetup;
+	AsyncBodySetup = nullptr;
+
+	RecreatePhysicsState();
+}
+
+void UPreviewMeshComponent::AbortPendingPhysicsCook()
+{
+	if (AsyncBodySetup)
+	{
+		UE_LOG(LogMegaMeshEditor, Display, TEXT("UPreviewMeshComponent::AbortPendingPhysicsCook aborting stale UBodySetup generation.  guid=%s"), *AsyncBodySetup->BodySetupGuid.ToString());
+		AsyncBodySetup->AbortPhysicsMeshAsyncCreation();
+		AsyncBodySetup = nullptr;
+	}
+}
+
+UBodySetup* UPreviewMeshComponent::GetBodySetup()
+{
+	if (!MeshBodySetup)
+	{
+		MeshBodySetup = CreateBodySetupHelper();
+	}
+
+	return MeshBodySetup;
+}
+
+UBodySetup* UPreviewMeshComponent::CreateBodySetupHelper()
+{
+	UBodySetup* NewBodySetup = NewObject<UBodySetup>(this, NAME_None, RF_NoFlags);
+	NewBodySetup->BodySetupGuid = FGuid::NewGuid();
+
+	NewBodySetup->bGenerateMirroredCollision = false;
+	NewBodySetup->bDoubleSidedGeometry = true;
+	NewBodySetup->CollisionTraceFlag = CTF_UseComplexAsSimple;
+
+	NewBodySetup->DefaultInstance.SetCollisionProfileName(UCollisionProfile::BlockAll_ProfileName);
+	NewBodySetup->bSupportUVsAndFaceRemap = false;
+
+	return NewBodySetup;
+}
+
+bool UPreviewMeshComponent::GetPhysicsTriMeshData(struct FTriMeshCollisionData* CollisionData, bool InUseAllTriData)
+{
+	if (MeshData.IsValid() && CVarCustomMeshGeneratePhysics.GetValueOnGameThread())
+	{
+		return MeshData->ConvertToTriMeshCollisionData(CollisionData);
+	}
+	return false; 
+}
+
+bool UPreviewMeshComponent::ContainsPhysicsTriMeshData(bool /*InUseAllTriData*/) const
+{
+	return CVarCustomMeshGeneratePhysics.GetValueOnGameThread() && MeshData.IsValid() && MeshData->VertexCount() > 0;
+}
+
+bool UPreviewMeshComponent::GetTriMeshSizeEstimates(struct FTriMeshCollisionDataEstimates& OutTriMeshEstimates, bool /*InUseAllTriData*/) const
+{
+	if (MeshData.IsValid())
+	{
+		OutTriMeshEstimates.VerticeCount = MeshData->VertexCount();
+	}
+	else
+	{
+		OutTriMeshEstimates.VerticeCount = 0;
+	}
+	return true;
+}
+
+void UPreviewMeshComponent::GetMeshId(FString& OutMeshId)
+{
+	// This is the opportunity to mark the physics inputs with an ID for ddc to cache the physics body.
+}
+
+UMaterialCacheVirtualTexture* GetMaterialCacheVirtualTexture(MeshPartition::UPreviewMeshComponent* Component, FGuid TagGuid)
+{
+	// Linear search
+	for (TObjectPtr MaterialCacheVirtualTexture : Component->MaterialCacheTextures)
+	{
+		UMaterialCacheVirtualTextureTag* TagHandle = MaterialCacheVirtualTexture->Tag.Get();
+		if ((TagHandle ? TagHandle->Guid : FGuid()) == TagGuid)
+		{
+			return MaterialCacheVirtualTexture;
+		}
+	}
+
+	return nullptr;
+}
+
+void UPreviewMeshComponent::SetMaterial(int32 ElementIndex, UMaterialInterface* Material)
+{
+	Super::SetMaterial(ElementIndex, Material);
+
+	if (!ensure(Material != nullptr))
+	{
+		return;
+	}
+
+	const FMaterialCachedExpressionData& ExpressionData = Material->GetCachedExpressionData();
+
+	for (TObjectPtr Tag : ExpressionData.MaterialCacheTags)
+	{
+		// May already have the tag
+		UMaterialCacheVirtualTextureTag* TagHandle = Tag.Get();
+		if (GetMaterialCacheVirtualTexture(this, TagHandle ? TagHandle->Guid : FGuid()))
+		{
+			continue;
+		}
+
+		static IConsoleVariable* CVarMaterialCacheTileCount = IConsoleManager::Get().FindConsoleVariable(TEXT("MegaMesh.MaterialCache.TileCount"));
+		const int32 TileCount = CVarMaterialCacheTileCount ? CVarMaterialCacheTileCount->GetInt() : 8;
+		
+		UMaterialCacheVirtualTexture* VirtualTexture = NewObject<UMaterialCacheVirtualTexture>(this);
+		VirtualTexture->OwningComponent = this;
+		VirtualTexture->Tag = Tag;
+		VirtualTexture->TileCount = FIntPoint(TileCount, TileCount);
+		VirtualTexture->UpdateResourceWithParams(UTexture::EUpdateResourceFlags::None);
+		MaterialCacheTextures.Add(VirtualTexture);
+	}
+}
+
+void UPreviewMeshComponent::GetUsedMaterials(TArray<UMaterialInterface*>& OutMaterials, bool bGetDebugMaterials) const
+{
+	Super::GetUsedMaterials(OutMaterials, bGetDebugMaterials);
+}
+
+class FMegaMeshCustomPreviewSceneProxy : public FPrimitiveSceneProxy
+{
+public:
+	FMegaMeshCustomPreviewSceneProxy(MeshPartition::UPreviewMeshComponent* InComponent)
+		: FPrimitiveSceneProxy(InComponent)
+		, VertexFactory(GetScene().GetFeatureLevel(), "MegaMeshCustomPreviewVertexFactory")
+		, MaterialRelevance(InComponent->GetMaterialRelevance(GetScene().GetShaderPlatform()))
+	{
+		if (ensure(InComponent->GetMeshData()->VertexCount() > 0))
+		{
+			InitResources(*InComponent->GetMeshData());
+
+			if (!InComponent->MaterialCacheTextures.IsEmpty() && InComponent->MaterialCacheTextures[0]->IsCurrentlyVirtualTextured())
+			{
+				for (UMaterialCacheVirtualTexture* MaterialCacheTexture : InComponent->MaterialCacheTextures)
+				{
+					MaterialCacheRenderProxies.Emplace(MaterialCacheTexture->CreateRenderProxy(0));
+				}
+			}
+
+			ENQUEUE_RENDER_COMMAND(FMegaMeshCustomPreviewSceneProxyInitialize)(
+				[&](FRHICommandListImmediate& RHICmdList)
+			{
+				PositionVertexBuffer.InitResource(RHICmdList);
+				StaticMeshVertexBuffer.InitResource(RHICmdList);
+				ColorVertexBuffer.InitResource(RHICmdList);
+
+				FLocalVertexFactory::FDataType Data;
+				PositionVertexBuffer.BindPositionVertexBuffer(&VertexFactory, Data);
+				StaticMeshVertexBuffer.BindTangentVertexBuffer(&VertexFactory, Data);
+				StaticMeshVertexBuffer.BindPackedTexCoordVertexBuffer(&VertexFactory, Data);
+				ColorVertexBuffer.BindColorVertexBuffer(&VertexFactory, Data);
+
+				VertexFactory.SetData(RHICmdList, Data);
+
+				VertexFactory.InitResource(RHICmdList);
+
+				
+				PositionVertexBuffer.InitResource(RHICmdList);
+				StaticMeshVertexBuffer.InitResource(RHICmdList);
+				ColorVertexBuffer.InitResource(RHICmdList);
+				IndexBuffer.InitResource(RHICmdList);
+
+#if RHI_RAYTRACING
+				InitRayTracingGeometry(RHICmdList);
+#endif // RHI_RAYTRACING
+			});
+		}
+
+		Material = InComponent->GetMaterial(0);
+
+		bSupportsRuntimeVirtualTexture = true;
+
+		if (FMaterial* MaterialResource = Material->GetMaterialResource(GShaderPlatformForFeatureLevel[GetScene().GetFeatureLevel()]))
+		{
+			bSupportsMaterialCache = MaterialResource->SamplesMaterialCache();
+		}
+	}
+
+	~FMegaMeshCustomPreviewSceneProxy()
+	{
+		VertexFactory.ReleaseResource();
+		PositionVertexBuffer.ReleaseResource();
+		StaticMeshVertexBuffer.ReleaseResource();
+		ColorVertexBuffer.ReleaseResource();
+		IndexBuffer.ReleaseResource();
+#if RHI_RAYTRACING
+		RayTracingGeometry.ReleaseResource();
+#endif // RHI_RAYTRACING
+	}
+
+	void InitResources(const FMeshData& Mesh)
+	{
+		constexpr int32 NumTexCoords = 1;
+
+		TArray<uint32>& Indices = IndexBuffer.Indices;
+
+		Tasks::FTask BuildIndexBuffer = Tasks::Launch(TEXT("BuildIndexBuffer"),
+			[&]()
+			{
+				Indices.Reserve(Mesh.MaxTriangleID() * 3);
+
+				TArray<TArray<int>> IndexGroups;
+				ParallelForWithTaskContext(IndexGroups, Mesh.MaxTriangleID(),
+				[&](TArray<int>& OutIndexGroup, int TriangleID)
+				{
+					if (!Mesh.IsTriangle(TriangleID))
+					{
+						return;
+					}
+					const Geometry::FIndex3i Triangle = Mesh.GetTriangle(TriangleID);
+					OutIndexGroup.Add(Triangle[0]);
+					OutIndexGroup.Add(Triangle[1]);
+					OutIndexGroup.Add(Triangle[2]);
+				});
+
+				for (TArray<int>& IndexGroup : IndexGroups)
+				{
+					Indices.Append(MoveTemp(IndexGroup));
+				}
+			});
+
+		Tasks::FTask BuildVertexBuffers = Tasks::Launch(TEXT("BuildVertexBuffers"),
+			[&]()
+			{
+				StaticMeshVertexBuffer.SetUseFullPrecisionUVs(true);
+				StaticMeshVertexBuffer.SetUseHighPrecisionTangentBasis(true);
+
+				PositionVertexBuffer.Init(Mesh.MaxVertexID());
+				StaticMeshVertexBuffer.Init(Mesh.MaxVertexID(), NumTexCoords);
+				ColorVertexBuffer.Init(Mesh.MaxVertexID());
+
+				ParallelFor(Mesh.MaxVertexID(), [&](int VertexID)
+				{
+					if (!Mesh.IsVertex(VertexID))
+					{
+						return;
+					}
+
+					const FVector3f Normal = FVector3f(Mesh.GetVertexNormal(VertexID));
+					PositionVertexBuffer.VertexPosition(VertexID) = FVector3f(Mesh.GetVertex(VertexID));
+					StaticMeshVertexBuffer.SetVertexTangents(VertexID, FVector3f(1, 0, 0), FVector3f(0, 1, 0), Normal);
+					ColorVertexBuffer.VertexColor(VertexID) = FColor::Black;
+			
+					StaticMeshVertexBuffer.SetVertexUV(VertexID, 0, Mesh.GetVertexUV(VertexID));
+				});
+			});
+
+		Tasks::Wait( MakeArrayView({ BuildIndexBuffer, BuildVertexBuffers }));
+	}
+
+	virtual void DrawStaticElements(FStaticPrimitiveDrawInterface* PDI) override
+	{
+		FMeshBatch BaseMeshBatch;
+		BaseMeshBatch.VertexFactory = &VertexFactory;
+		BaseMeshBatch.MaterialRenderProxy = Material ? Material->GetRenderProxy() : UMaterial::GetDefaultMaterial(MD_Surface)->GetRenderProxy();
+		BaseMeshBatch.ReverseCulling = IsLocalToWorldDeterminantNegative();
+		BaseMeshBatch.Type = PT_TriangleList;
+		BaseMeshBatch.DepthPriorityGroup = SDPG_World;
+		BaseMeshBatch.LODIndex = 0;
+		BaseMeshBatch.SegmentIndex = 0;
+		BaseMeshBatch.MeshIdInPrimitive = 0;
+		BaseMeshBatch.LCI = nullptr;
+		BaseMeshBatch.CastShadow = true;
+		BaseMeshBatch.bUseForMaterial = true;
+		BaseMeshBatch.bDitheredLODTransition = false;
+		BaseMeshBatch.bUseForDepthPass = true;
+		BaseMeshBatch.bUseAsOccluder = ShouldUseAsOccluder();
+
+		FMeshBatchElement& BatchElement = BaseMeshBatch.Elements[0];
+		
+		BatchElement.IndexBuffer = &IndexBuffer;
+		BatchElement.PrimitiveUniformBuffer = GetUniformBuffer();
+		BatchElement.NumPrimitives = IndexBuffer.Indices.Num() / 3;
+		BatchElement.FirstIndex = 0;
+		BatchElement.MinVertexIndex = 0;
+		BatchElement.MaxVertexIndex = PositionVertexBuffer.GetNumVertices() - 1;
+
+		PDI->DrawMesh(BaseMeshBatch, FLT_MAX);
+
+		{
+			FMeshBatch VTMeshBatch(BaseMeshBatch);
+			VTMeshBatch.CastShadow = false;
+			VTMeshBatch.bUseAsOccluder = false;
+			VTMeshBatch.bUseForDepthPass = false;
+			VTMeshBatch.bUseForMaterial = false;
+			VTMeshBatch.bDitheredLODTransition = false;
+			VTMeshBatch.bRenderToVirtualTexture = true;
+			for (ERuntimeVirtualTextureMaterialType MaterialType : RuntimeVirtualTextureMaterialTypes)
+			{
+				VTMeshBatch.RuntimeVirtualTextureMaterialType = (uint32)MaterialType;
+				PDI->DrawMesh(VTMeshBatch, FLT_MAX);
+			}
+		}
+	}
+
+	virtual void GetDynamicMeshElements(const TArray<const FSceneView*>& Views, const FSceneViewFamily& ViewFamily, uint32 VisibilityMap, FMeshElementCollector& Collector) const override
+	{
+		const FEngineShowFlags& EngineShowFlags = ViewFamily.EngineShowFlags;
+		const bool bWireframe = (AllowDebugViewmodes() && EngineShowFlags.Wireframe);
+		const bool bIsSelected = IsSelected();
+
+		FMaterialRenderProxy* WireframeMaterialProxy = nullptr;
+		if (bWireframe)
+		{
+			const FLinearColor WireframeColorOverride = bIsSelected ? GEngine->GetSelectedMaterialColor() : GetWireframeColor();
+			WireframeMaterialProxy = new FColoredMaterialRenderProxy(GEngine->WireframeMaterial->GetRenderProxy(), WireframeColorOverride);
+			Collector.RegisterOneFrameMaterialProxy(WireframeMaterialProxy);
+		}
+
+		for (int32 ViewIndex = 0; ViewIndex < Views.Num(); ViewIndex++)
+		{
+			if (VisibilityMap & (1 << ViewIndex))
+			{
+				FMaterialRenderProxy* MaterialProxy = Material ? Material->GetRenderProxy() : UMaterial::GetDefaultMaterial(MD_Surface)->GetRenderProxy();
+
+				FDynamicPrimitiveUniformBuffer& DynamicPrimitiveUniformBuffer = Collector.AllocateOneFrameResource<FDynamicPrimitiveUniformBuffer>();
+				FPrimitiveUniformShaderParametersBuilder Builder;
+				BuildUniformShaderParameters(Builder);
+				DynamicPrimitiveUniformBuffer.Set(Collector.GetRHICommandList(), Builder);
+
+				if (IndexBuffer.Indices.Num() > 0)
+				{
+					FMeshBatch& MeshBatch = Collector.AllocateMesh();
+					FMeshBatchElement& BatchElement = MeshBatch.Elements[0];
+					BatchElement.IndexBuffer = &IndexBuffer;
+					MeshBatch.bWireframe = bWireframe;
+					MeshBatch.bDisableBackfaceCulling = bWireframe;
+					MeshBatch.VertexFactory = &VertexFactory;
+					MeshBatch.MaterialRenderProxy = bWireframe ? WireframeMaterialProxy : MaterialProxy;
+					MeshBatch.bRenderToVirtualTexture = true;
+
+					BatchElement.PrimitiveUniformBufferResource = &DynamicPrimitiveUniformBuffer.UniformBuffer;
+
+					BatchElement.FirstIndex = 0;
+					BatchElement.NumPrimitives = IndexBuffer.Indices.Num() / 3;
+					BatchElement.MinVertexIndex = 0;
+					BatchElement.MaxVertexIndex = PositionVertexBuffer.GetNumVertices() - 1;
+					MeshBatch.ReverseCulling = IsLocalToWorldDeterminantNegative();
+					MeshBatch.Type = PT_TriangleList;
+					MeshBatch.DepthPriorityGroup = SDPG_World;
+					Collector.AddMesh(ViewIndex, MeshBatch);
+				}
+			}
+		}
+	}
+	
+#if RHI_RAYTRACING
+	void InitRayTracingGeometry(FRHICommandListBase& RHICmdList)
+	{
+		if (!IsRayTracingEnabled())
+		{
+			return;
+		}
+
+		RayTracingGeometry.ReleaseResource();
+		
+		FRayTracingGeometryInitializer Initializer;
+		Initializer.IndexBuffer = IndexBuffer.IndexBufferRHI;
+		Initializer.TotalPrimitiveCount = IndexBuffer.Indices.Num() / 3;
+		Initializer.GeometryType = RTGT_Triangles;
+		Initializer.bFastBuild = true;
+		Initializer.bAllowUpdate = false;
+
+		FRayTracingGeometrySegment Segment;
+		Segment.VertexBuffer = PositionVertexBuffer.VertexBufferRHI;
+		Segment.NumPrimitives = Initializer.TotalPrimitiveCount;
+		Segment.MaxVertices = PositionVertexBuffer.GetNumVertices();
+
+		Initializer.Segments.Add(Segment);
+
+		RayTracingGeometry.SetInitializer(MoveTemp(Initializer));
+		RayTracingGeometry.InitResource(RHICmdList);
+	}
+
+	virtual bool IsRayTracingRelevant() const override
+	{
+		return true;
+	}
+
+	virtual bool HasRayTracingRepresentation() const override
+	{
+		return true;
+	}
+
+	virtual void GetDynamicRayTracingInstances(FRayTracingInstanceCollector& Collector) override
+	{
+		TConstArrayView<const FSceneView*> Views = Collector.GetViews();
+		const uint32 VisibilityMap = Collector.GetVisibilityMap();
+
+		// RT geometry will be generated based on first active view and then reused for all other views
+		// TODO: Expose a way for developers to control whether to reuse RT geometry or create one per-view
+		const int32 FirstActiveViewIndex = FMath::CountTrailingZeros(VisibilityMap);
+		checkf(Views.IsValidIndex(FirstActiveViewIndex), TEXT("There should be at least one active view when calling GetDynamicRayTracingInstances(...)."));
+
+		const FSceneView* FirstActiveView = Views[FirstActiveViewIndex];
+
+		FRayTracingInstance RayTracingInstance;
+		RayTracingInstance.Geometry = &RayTracingGeometry;
+		RayTracingInstance.InstanceTransforms.Add(GetLocalToWorld());
+	
+		FMeshBatch MeshBatch;
+	
+		MeshBatch.VertexFactory = &VertexFactory;
+		MeshBatch.SegmentIndex = 0;
+		MeshBatch.MaterialRenderProxy = Material ? Material->GetRenderProxy() : UMaterial::GetDefaultMaterial(MD_Surface)->GetRenderProxy();
+		MeshBatch.Type = PT_TriangleList;
+		MeshBatch.DepthPriorityGroup = SDPG_World;
+		MeshBatch.CastRayTracedShadow = IsShadowCast(FirstActiveView);
+	
+		FMeshBatchElement& BatchElement = MeshBatch.Elements[0];
+		BatchElement.IndexBuffer = &IndexBuffer;
+		BatchElement.PrimitiveUniformBuffer = GetUniformBuffer();
+		BatchElement.FirstIndex = 0;
+		BatchElement.NumPrimitives = IndexBuffer.Indices.Num() / 3;
+		BatchElement.MinVertexIndex = 0;
+		BatchElement.MaxVertexIndex = PositionVertexBuffer.GetNumVertices() - 1;
+	
+		RayTracingInstance.Materials.Add(MeshBatch);
+
+		for (int32 ViewIndex = 0; ViewIndex < Views.Num(); ViewIndex++)
+		{
+			if ((VisibilityMap & (1 << ViewIndex)) == 0)
+			{
+				continue;
+			}
+
+			Collector.AddRayTracingInstance(ViewIndex, RayTracingInstance);
+		}
+	}
+#endif // RHI_RAYTRACING
+
+	virtual FPrimitiveViewRelevance GetViewRelevance(const FSceneView* View) const override
+	{
+		const FEngineShowFlags& EngineShowFlags = View->Family->EngineShowFlags;
+
+		FPrimitiveViewRelevance Result;
+		Result.bDrawRelevance = IsShown(View);
+		Result.bShadowRelevance = IsShadowCast(View);
+		Result.bStaticRelevance = !EngineShowFlags.Wireframe;
+		Result.bDynamicRelevance = EngineShowFlags.Wireframe;
+		Result.bEditorVisualizeLevelInstanceRelevance = IsEditingLevelInstanceChild();
+		Result.bEditorStaticSelectionRelevance = (IsSelected() || IsHovered());
+		Result.bRenderInMainPass = ShouldRenderInMainPass();
+		Result.bUsesLightingChannels = GetLightingChannelMask() != GetDefaultLightingChannelMask();
+		Result.bTranslucentSelfShadow = bCastVolumetricTranslucentShadow;
+		Result.bRenderCustomDepth = ShouldRenderCustomDepth();
+
+		MaterialRelevance.SetPrimitiveViewRelevance(Result);
+
+		Result.bVelocityRelevance = DrawsVelocity() && Result.bOpaque && Result.bRenderInMainPass;
+
+		return Result;
+	}
+	
+	virtual SIZE_T GetTypeHash() const override
+	{
+		static size_t UniquePointer;
+		return reinterpret_cast<size_t>(&UniquePointer);
+	}
+
+
+	virtual HHitProxy* CreateHitProxies(UPrimitiveComponent* Component, TArray<TRefCountPtr<HHitProxy> >& OutHitProxies) override
+	{
+		MeshPartition::HPreviewProxy* HProxy = new MeshPartition::HPreviewProxy(Component);
+		OutHitProxies.Add(HProxy);
+		return HProxy;
+	}
+
+	virtual uint32 GetMemoryFootprint(void) const override { return(sizeof(*this) + GetAllocatedSize()); }
+	uint32 GetAllocatedSize(void) const { return(FPrimitiveSceneProxy::GetAllocatedSize()); }
+
+private:
+	FLocalVertexFactory VertexFactory;
+
+	FMaterialRelevance MaterialRelevance;
+	UMaterialInterface* Material = nullptr;
+
+	FStaticMeshVertexBuffer StaticMeshVertexBuffer;
+	FPositionVertexBuffer PositionVertexBuffer;
+	FColorVertexBuffer ColorVertexBuffer;
+	
+	FDynamicMeshIndexBuffer32 IndexBuffer;
+
+#if RHI_RAYTRACING
+	FRayTracingGeometry RayTracingGeometry;
+#endif // RHI_RAYTRACING
+};
+
+FPrimitiveSceneProxy* UPreviewMeshComponent::CreateSceneProxy()
+{
+	if (MeshData->VertexCount() == 0)
+	{
+		return nullptr;
+	}
+	return new FMegaMeshCustomPreviewSceneProxy(this);
+}
+
+HPreviewProxy::HPreviewProxy(const UActorComponent* InComponent)
+	: HComponentVisProxy(InComponent, EHitProxyPriority::HPP_World)
+	, Actor(InComponent ? InComponent->GetOwner() : nullptr)
+{}
+
+IMPLEMENT_HIT_PROXY(MeshPartition::HPreviewProxy, HComponentVisProxy);
+} // namespace UE::MeshPartition
\ No newline at end of file
diff --git a/Engine/Plugins/Experimental/MeshPartition/Source/MeshPartitionEditor/Private/MeshPartitionSeparateWorldBuilder.cpp b/Engine/Plugins/Experimental/MeshPartition/Source/MeshPartitionEditor/Private/MeshPartitionSeparateWorldBuilder.cpp
new file mode 100644
--- /dev/null
+++ b/Engine/Plugins/Experimental/MeshPartition/Source/MeshPartitionEditor/Private/MeshPartitionSeparateWorldBuilder.cpp
@@ -0,0 +1,297 @@
+// Copyright Epic Games, Inc. All Rights Reserved.
+
+#include "MeshPartitionSeparateWorldBuilder.h"
+#include "MeshPartition.h"
+#include "MeshPartitionEditorModule.h"
+#include "MeshPartitionEditorSubsystem.h"
+#include "MeshPartitionCompiledSection.h"
+#include "MeshPartitionEditorComponent.h"
+#include "MeshPartitionModifierComponent.h"
+#include "MeshPartitionDescriptorCache.h"
+#include "WorldPartition/WorldPartitionSubsystem.h"
+#include "WorldPartition/WorldPartitionHandle.h"
+#include "WorldPartition/WorldPartition.h"
+#include "WorldPartition/WorldPartitionHelpers.h"
+#include "MeshPartitionWorldPartitionHelpers.h"
+#include "EditorWorldUtils.h"
+#include "EngineUtils.h"
+#include "UObject/Package.h"
+
+namespace UE::MeshPartition
+{
+void FSeparateWorldBuilder::InitializeToBuildOutOfWorldPackage(const FString& InWorldPackageName)
+{
+	// unload existing world if necessary
+	if (World != nullptr)
+	{
+		FString ExistingWorldPackageName = World->GetPackage()->GetPathName();
+		if (ExistingWorldPackageName != InWorldPackageName)
+		{
+			UE_LOG(LogMegaMeshEditor, Verbose, TEXT("InitializeToBuildOutOfWorldPackage: Changing Worlds '%s' ==> '%s'"), *ExistingWorldPackageName, *InWorldPackageName);
+
+			// delete existing world
+			World->MarkObjectsPendingKill();
+			World->DestroyWorld(false);
+			World->RemoveFromRoot();
+			World = nullptr;
+
+			// ensure it will be garbage collected when we are done...
+			// TODO: we might want to run our own GC as by default GC only runs at PIE transition / save / load in editor
+			check(!BuilderPackage->HasAllFlags(EObjectFlags::RF_MarkAsRootSet));
+		}
+		else
+		{
+			UE_LOG(LogMegaMeshEditor, Verbose, TEXT("InitializeToBuildOutOfWorldPackage: Using Existing World '%s' == '%s'"), *ExistingWorldPackageName, *InWorldPackageName);
+		}
+	}
+	else
+	{
+		UE_LOG(LogMegaMeshEditor, Verbose, TEXT("InitializeToBuildOutOfWorldPackage: Loading New World '%s'"), *InWorldPackageName);
+	}
+
+	if (BuilderPackage == nullptr)
+	{
+		BuilderPackage = CreatePackage(TEXT("/Temp/MegaMeshBuilderPackage"));
+	}
+	ensure(BuilderPackage);
+
+	if (World == nullptr)
+	{
+		const TCHAR* LongPackageName = *InWorldPackageName;
+		ELoadFlags LoadFlags = (ELoadFlags)
+				( ELoadFlags::LOAD_None
+				| ELoadFlags::LOAD_NoWarn
+				| ELoadFlags::LOAD_EditorOnly
+				| ELoadFlags::LOAD_DisableCompileOnLoad);
+		FArchive* ReaderArchive = nullptr;
+		FLinkerInstancingContext* LinkerInstancingContext = nullptr;
+		UPackage* MapPackage = LoadPackage(BuilderPackage, LongPackageName, LoadFlags, ReaderArchive, LinkerInstancingContext);
+		check(MapPackage);
+
+		World = UWorld::FindWorldInPackage(MapPackage);
+		check(!World->IsPlayInEditor());
+	}
+	ensure(World);
+}
+
+void FSeparateWorldBuilder::BuildCompiledSectionMesh(MeshPartition::ACompiledSection* CompiledSection, FName BuildVariantName)
+{
+	// NOTE: the compiled section is the actor we need to build the mesh for
+	// however, the compiled section does NOT belong to our World (our World is a temp world loaded just for building)
+	check(CompiledSection);
+	UWorld* CompiledSectionWorld = CompiledSection->GetWorld();
+
+	MeshPartition::FCompiledSectionBuildInfo BuildInfo = CompiledSection->GetBuildInfo();
+
+	// Get the megamesh definition.  Note that the compiled section is in a streaming cell, and it's parent megamesh may not be loaded
+	// So instead we make sure to grab the definition from the path stored on the buildinfo
+	FSoftObjectPath MegaMeshDefinitionPath(BuildInfo.MegaMeshDefinitionPath);
+	TSoftObjectPtr<UMeshPartitionDefinition> MegaMeshDefinition(MegaMeshDefinitionPath);
+	UMeshPartitionDefinition* Definition = MegaMeshDefinition.Get();
+	if (Definition == nullptr)
+	{
+		Definition = UMeshPartitionDefinition::StaticClass()->GetDefaultObject<UMeshPartitionDefinition>();
+	}
+
+	const MeshPartition::FCompiledSectionBuildVariant& BuildVariant = Definition->GetCompiledSectionBuildVariantByName(BuildVariantName);
+
+	if (MeshPartition::FMeshBuilder* Builder = UMeshPartitionEditorSubsystem::GetMeshBuilder())
+	{
+		MeshPartition::FModifierDescriptorCache* DescriptorCache = UMeshPartitionEditorSubsystem::GetDescriptorCache();
+		check(DescriptorCache);
+
+		// We search by GUID because the path can be modified by PIE
+		TSharedPtr<const MeshPartition::FCachedDescriptors> Descriptors = DescriptorCache->GetCachedDescriptorsByGuid(BuildInfo.MegaMeshGUID);
+
+		FSoftObjectPath MeshPartitionActorPath = Descriptors->GetMegaMeshPath();
+
+		// build the section group that corresponds to these base modifiers (after first fixing the PIE path names..)
+		MeshPartition::EditorUtils::FPIEPathFixer& PIEPathFixer = UMeshPartitionEditorSubsystem::GetPIEPathFixer();
+		for (FSoftObjectPath& ModifierPath : BuildInfo.BaseModifierPaths)
+		{
+			PIEPathFixer.FixInPlace(ModifierPath);
+		}
+		MeshPartition::FModifierGroup ModifierGroup = Descriptors->BuildModifierGroupForBaseModifierPaths(BuildInfo.BaseModifierPaths);
+
+		// find the mesh partition actor by GUID (NOTE: we can't find it by path, unless we use the PIE world path, otherwise it will get the mesh partition actor in the main editor world)
+		FGuid MeshPartitionActorGUID = Descriptors->GetMegaMeshGUID();
+		const AMeshPartition* MeshPartitionActor = nullptr;
+		for (AMeshPartition* MeshPartition : TActorRange<AMeshPartition>(CompiledSectionWorld))
+		{
+			if (MeshPartition && (MeshPartition->GetActorGuid() == MeshPartitionActorGUID))
+			{
+				MeshPartitionActor = MeshPartition;
+				break;
+			}
+		}
+
+		if (MeshPartitionActor == nullptr)
+		{
+			UE_LOG(LogMegaMeshEditor, Warning, TEXT("Failed to build compiled section: could not find parent MeshPartition actor %s in World %s"), *MeshPartitionActorGUID.ToString(), *CompiledSectionWorld->GetName());
+			return;
+		}
+
+		UMeshPartitionEditorComponent* MeshPartitionEditorComponent = Cast<UMeshPartitionEditorComponent>(MeshPartitionActor->GetMeshPartitionComponent());
+		check(MeshPartitionEditorComponent);
+		check(MeshPartitionEditorComponent->GetWorld() == CompiledSectionWorld);
+
+		auto OnActorLoaded = [](AActor* Actor, bool bIsInLevelInstance)
+			{
+				// In the commandlet build path we just record the actor refs to ensure they are not GCed while we are working.
+				// As this path is synchronous currently, it's not necessary,
+			};
+
+		auto OnModifierLoaded = [](MeshPartition::UModifierComponent* Modifier, bool bIsInLevelInstance)
+			{
+				Modifier->InitializeModifier();
+			};
+
+		TSet<FWorldPartitionReference> ActorRefs;
+		WorldPartitionHelpers::LoadAllRelevantMegaMeshActors(ModifierGroup, MeshPartitionActor, ActorRefs, OnActorLoaded, OnModifierLoaded);
+
+		FGuid GroupHash = ModifierGroup.UpdateAndComputeModifierGroupHash();
+		BuildInfo.ModifiersHash = GroupHash;
+
+		MeshPartition::FBuilderSettings BuilderSettings;
+		{
+			BuilderSettings.BuildType = MeshPartition::EBuildType::CompiledSection;
+			BuilderSettings.Transform = FTransform::Identity;
+			
+			// Fill out modifiers to process
+			static bool bFillModifiersToProcess = true;
+			BuilderSettings.ModifiersToProcess.Empty();
+			
+			if (bFillModifiersToProcess)
+			{
+				ModifierGroup.ForAllModifiers([&BuilderSettings](MeshPartition::UModifierComponent* ModifierComponent)
+				{
+					if (ensure(ModifierComponent != nullptr))
+					{
+						BuilderSettings.ModifiersToProcess.Add(ModifierComponent);
+					}
+
+					return true;
+				});
+			}
+
+			BuilderSettings.TypePriorities = Definition->GetModifierTypePriorities();
+			// BuilderSettings.FilterBounds = ;				// optional bounds
+			// BuilderSettings.ModifierFilter = ...;		// Optional filter function for the list of modifiers to process.
+			BuilderSettings.MaxSectionComplexity = BuildVariant.MaxSectionComplexity;
+			BuilderSettings.bRecomputeNormals = false;
+			BuilderSettings.bRecomputeTangents = false;		// not yet supported
+			BuilderSettings.bCacheResult = true;
+			BuilderSettings.bBuildSpatial = false;
+			BuilderSettings.bAllowDDCRead = true;			// Should builds with these settings be allowed to use the DDC cache to avoid processing.
+			BuilderSettings.bAllowDDCWrite = true;			// Should builds produced with these settings be allowed to store their results in DDC.
+			BuilderSettings.TexcoordGenerationOptions = FChannelCollectionUVLayoutOptions::GetFromDefinition(Definition);
+			// Choose the higher quality layout method for compiled sections
+			BuilderSettings.TexcoordGenerationOptions->UVLayoutMethod = EChannelCollectionUVLayoutMethod::ReferenceBoxProject;
+
+			MeshPartition::FBuilderSettings::FChannelRenderSettings ChannelRenderSettings;
+			ChannelRenderSettings.ChannelMap = Definition->GetChannelMap();
+			ChannelRenderSettings.TexelSize = Definition->GetChannelTexelSize();
+			BuilderSettings.ChannelRenderSettings.Emplace(MoveTemp(ChannelRenderSettings));
+		}
+
+		const bool bAllowDDC = false;
+		MeshPartition::FBuildTaskHandle TaskHandle = Builder->Build(BuilderSettings, ModifierGroup, bAllowDDC);
+
+		//TaskHandles.Add(TaskHandle);
+		TaskHandle.Wait();
+		check(TaskHandle.IsCompleted());
+
+		TSharedPtr<const MeshPartition::FSectionChannels> SectionChannel = TaskHandle.GetTask()->GetSectionChannels();
+
+		const bool bUseStaticMobility = true;
+		CompiledSection = MeshPartitionEditorComponent->PrepareMegaMeshCompiledSection(BuildInfo, BuildVariant, CompiledSection, bUseStaticMobility);
+
+		CompiledSection->SetChannelTexture(SectionChannel->Texture.Get());
+		CompiledSection->SetChannelTexcoordDesc(SectionChannel->TexcoordMetrics);
+		CompiledSection->SetChannelTable(SectionChannel->Table);
+
+		TArray<TWeakObjectPtr<MeshPartition::UModifierComponent>> ModifierPtrs = ModifierGroup.AllResolvedModifierPtrs();
+
+		TSharedPtr<const FMeshData> ProcessedMesh = TaskHandle.GetTask()->GetMesh();
+		const FMeshData* ProcessedMeshPtr = ProcessedMesh.Get();
+
+		if (ensure(ProcessedMeshPtr != nullptr))
+		{
+			MeshPartitionEditorComponent->PostBuildSectionMesh(CompiledSection, *ProcessedMeshPtr, ModifierPtrs);
+
+			const int32 Vertices = ProcessedMesh->VertexCount();
+			const int32 Tris = ProcessedMesh->TriangleCount();
+
+			// TODO: this can be moved async (if we track intermediate state, and ensure GC doesn't reclaim our stuff)
+
+			MeshPartition::FTransformerUnit InitialUnit;
+
+			InitialUnit.Section = CompiledSection;
+			InitialUnit.MeshData = ProcessedMesh;
+			InitialUnit.bShouldRecomputeNormals = true;
+			InitialUnit.bShouldRecomputeTangents = true;
+
+			TUniquePtr<MeshPartition::FTransformerContext> TransformerContext = MeshPartitionEditorComponent->LaunchTransformers(MoveTemp(InitialUnit), Definition, BuildVariant);
+			
+			if (TransformerContext.IsValid() && TransformerContext->JoinTask.IsValid() && !TransformerContext->JoinTask.IsCompleted())
+			{
+				WaitOnGameThread(*TransformerContext);
+			}
+
+			MeshPartitionEditorComponent->PostProcessSection(CompiledSection, ModifierPtrs);
+		
+			CompiledSection->SetIsPlaceholder(false);
+
+			UE_LOG(LogMegaMeshEditor, Verbose, TEXT("Built Compiled Section: %p %s (%d bases, %d vertices, %d tris)"), CompiledSection, *CompiledSection->GetActorNameOrLabel(), BuildInfo.BaseModifierPaths.Num(), Vertices, Tris);
+		}
+	}
+}
+
+void FSeparateWorldBuilder::RequestCompiledSectionBuild(MeshPartition::ACompiledSection* CompiledSection)
+{
+	if (ensure(CompiledSection->IsPlaceholder()))
+	{
+		check(IsInGameThread());
+		CompiledSectionsToBuild.Add(CompiledSection);
+	}
+}
+
+void FSeparateWorldBuilder::Tick()
+{
+	check(IsInGameThread());
+	int32 SectionIndex = 0;
+	MeshPartition::EditorUtils::FPIEPathFixer& PIEPathFixer = UMeshPartitionEditorSubsystem::GetPIEPathFixer();
+	for (TWeakObjectPtr<MeshPartition::ACompiledSection> CompiledSectionPtr : CompiledSectionsToBuild)
+	{
+		UE_LOG(LogMegaMeshEditor, Verbose, TEXT("Building Placeholder Compiled Section (%d/%d) ..."), SectionIndex, CompiledSectionsToBuild.Num());
+
+		if (MeshPartition::ACompiledSection* CompiledSection = CompiledSectionPtr.Get())
+		{
+			if (CompiledSection->IsPlaceholder())
+			{
+				FSoftObjectPath MeshPartitionPath = CompiledSection->GetBuildInfo().MegaMeshPath;
+				check(MeshPartitionPath.IsValid());
+
+				// MeshPartition path may be a PIE path -- we need to fix it to get the NonPIE package name
+				PIEPathFixer.FixInPlace(MeshPartitionPath);
+				
+				// TODO: we should ideally sort requests to avoid changing worlds as much as possible..
+				// this could be terrible if requests are coming from multiple worlds simultaneously.
+				// For now it should be ok as we generally work on one world at a time.
+				InitializeToBuildOutOfWorldPackage(MeshPartitionPath.GetAssetPath().ToString());
+
+				FName BuildVariantName = CompiledSection->GetBuildInfo().BuildVariantName;
+				BuildCompiledSectionMesh(CompiledSection, BuildVariantName);
+			}
+		}
+
+		UE_LOG(LogMegaMeshEditor, Verbose, TEXT("BUILT Placeholder Compiled Section (%d/%d)"), SectionIndex, CompiledSectionsToBuild.Num());
+		SectionIndex++;
+	}
+	CompiledSectionsToBuild.Empty();
+}
+
+FSeparateWorldBuilder::~FSeparateWorldBuilder()
+{
+}
+} // namespace UE::MeshPartition
\ No newline at end of file
diff --git a/Engine/Plugins/Experimental/MeshPartition/Source/MeshPartitionEditor/Private/MeshPartitionStaticMeshTransformer.cpp b/Engine/Plugins/Experimental/MeshPartition/Source/MeshPartitionEditor/Private/MeshPartitionStaticMeshTransformer.cpp
new file mode 100644
--- /dev/null
+++ b/Engine/Plugins/Experimental/MeshPartition/Source/MeshPartitionEditor/Private/MeshPartitionStaticMeshTransformer.cpp
@@ -0,0 +1,304 @@
+// Copyright Epic Games, Inc. All Rights Reserved.
+
+#include "MeshPartitionStaticMeshTransformer.h"
+
+#include "Engine/StaticMesh.h"
+#include "Materials/MaterialInstanceConstant.h"
+#include "MeshPartitionDependencyInterface.h"
+#include "MeshPartitionEditorComponent.h"
+#include "MeshPartitionEditorModule.h"
+#include "MeshPartitionEditorUtils.h"
+#include "MeshPartitionCompiledSection.h"
+#include "MeshPartitionMeshBuilder.h"
+#include "MeshPartitionPreviewSection.h"
+
+namespace UE::MeshPartition
+{
+
+namespace MeshPartitionStaticMeshTransformerLocals
+{
+
+} // namespace MeshPartitionStaticMeshTransformerLocals
+
+bool FStaticMeshTransformer::Execute(MeshPartition::FTransformerContext& InTransformerContext) const
+{
+	UE::Tasks::TTask<TArray<UStaticMesh*>> CreateStaticMeshesTask = UE::Tasks::Launch(TEXT("FStaticMeshTransformer::CreateStaticMeshes"), [this,
+		&InTransformerContext]() mutable
+		{
+			return CreateStaticMeshes(InTransformerContext);
+		},
+		UE::Tasks::ETaskPriority::Normal,
+		UE::Tasks::EExtendedTaskPriority::GameThreadNormalPri
+	);
+
+	TArray<UE::Tasks::FTask> BuildStaticMeshTasks;
+
+	BuildStaticMeshTasks.SetNum(InTransformerContext.TransformerUnits.Num());
+
+	for (int32 Index = 0; Index < InTransformerContext.TransformerUnits.Num(); ++Index)
+	{
+		BuildStaticMeshTasks[Index] = UE::Tasks::Launch(TEXT("FStaticMeshTransformer::BuildMesh"), [this,
+			&InTransformerContext,
+			&CreateStaticMeshesTask,
+			Index]() mutable
+			{
+				BuildStaticMesh(InTransformerContext, CreateStaticMeshesTask, Index);
+			}, UE::Tasks::Prerequisites(CreateStaticMeshesTask));
+	}
+
+	UE::Tasks::FTask FinalizeStaticMeshesTask = UE::Tasks::Launch(TEXT("FStaticMeshTransformer::FinalizeStaticMeshes"), [this,
+		&CreateStaticMeshesTask,
+		&InTransformerContext]() mutable
+		{
+			FinalizeStaticMeshes(InTransformerContext, CreateStaticMeshesTask);
+		},
+		UE::Tasks::Prerequisites(BuildStaticMeshTasks),
+		UE::Tasks::ETaskPriority::Normal,
+		UE::Tasks::EExtendedTaskPriority::GameThreadNormalPri
+	);
+
+	FinalizeStaticMeshesTask.Wait();
+
+	return true;
+}
+
+void FStaticMeshTransformer::GatherDependencies(MeshPartition::IDependencyInterface& InDependencies) const
+{
+	InDependencies += CollisionProfile.Name;
+	
+	for (const MeshPartition::FLODSettings& LOD : LODs)
+	{
+		InDependencies += LOD.bReduceMesh;
+		
+		if (LOD.bReduceMesh)
+		{
+			InDependencies.AddUStructDependencyViaReflection(LOD.ReductionSettings);
+		}
+	}
+
+	InDependencies += LODStreaming;
+	InDependencies += NumStreamedLODs;
+	InDependencies += bUseNanite;
+	InDependencies += bCanEverAffectNavigation;
+}
+
+TArray<UStaticMesh*> FStaticMeshTransformer::CreateStaticMeshes(MeshPartition::FTransformerContext& InTransformerContext) const
+{
+	check(IsInGameThread());
+		   	
+	TArray<UStaticMesh*> Results;
+		   	
+	if (InTransformerContext.bWasCancelled)
+	{
+		return Results;
+	}
+		   	
+	for (const MeshPartition::FTransformerUnit& TransformerUnit : InTransformerContext.TransformerUnits)
+	{
+		if (!ensure(TransformerUnit.MeshData.IsValid()) || !ensure(TransformerUnit.MeshData->VertexCount() != 0) || !ensure(TransformerUnit.Section != nullptr))
+		{
+			// Still add nullptr to preserve 1-1 num between TransformerUnits and StaticMeshes
+			Results.Add(nullptr);
+			continue;
+		}
+
+		UStaticMesh* StaticMesh = MeshPartition::EditorUtils::CreateStaticMesh(TransformerUnit.Section, TEXT("MeshPartitionStaticMesh"));
+		StaticMesh->SetInternalFlags(EInternalObjectFlags::Async);
+		   	
+		if (!bCanEverAffectNavigation)
+		{
+			StaticMesh->MarkAsNotHavingNavigationData();
+		}
+
+		Results.Add(StaticMesh);
+	}
+		   	
+	return Results;
+}
+
+void FStaticMeshTransformer::BuildStaticMesh(MeshPartition::FTransformerContext& InTransformerContext, UE::Tasks::TTask<TArray<UStaticMesh*>>& InCreateStaticMeshesTask, const int32 InTransformerUnitIndex) const
+{
+	if ((InTransformerContext.bWasCancelled) || (InTransformerContext.TransformerUnits.Num() <= InTransformerUnitIndex))
+	{
+		return;
+	}
+			
+	UStaticMesh* StaticMesh = nullptr;
+	const MeshPartition::FTransformerUnit& TransformerUnit = InTransformerContext.TransformerUnits[InTransformerUnitIndex];
+	
+	{
+		TArray<UStaticMesh*>& StaticMeshes = InCreateStaticMeshesTask.GetResult();
+					
+		if (InTransformerUnitIndex >= StaticMeshes.Num())
+		{
+			return;
+		}
+					
+		StaticMesh = StaticMeshes[InTransformerUnitIndex];
+	}
+
+	// static mesh can be null if the built section had no data, we would not have created the static mesh for it.
+	if (StaticMesh == nullptr)
+	{
+		return;
+	}
+
+	if (!ensure(TransformerUnit.MeshData != nullptr))
+	{
+		UE_LOG(LogMegaMeshEditor, Error, TEXT("Cannot build static mesh for nullptr MeshData."));
+		return;
+	}
+			
+	if (!ensure(IsValid(StaticMesh)))
+	{
+		UE_LOG(LogMegaMeshEditor, Error, TEXT("Cannot get StaticMesh."));
+		return;
+	}
+	
+	if (LODs.Num() == 0)
+	{
+		StaticMesh->SetNumSourceModels(1);
+		StaticMesh->CreateMeshDescription(0);
+		FStaticMeshSourceModel& SourceModel = StaticMesh->GetSourceModel(0);
+		 	
+		MeshPartition::EditorUtils::BuildSourceModel(SourceModel, *TransformerUnit.MeshData, TransformerUnit.bShouldRecomputeNormals, TransformerUnit.bShouldRecomputeTangents);
+		return;
+	}
+
+	int32 LODCount = (LODs.Num() > 0) ? LODs.Num() : 1;
+	check(LODCount > 0);
+	
+	StaticMesh->SetNumSourceModels(LODCount);
+
+	for (int32 LODIndex = 0; LODIndex < LODCount; ++LODIndex)
+	{
+		StaticMesh->CreateMeshDescription(LODIndex);
+	}
+
+	// Any full resolution LOD _should_ be the first LOD index, but, just in case we can search for it.
+	int32 BaseLODIndex = -1;
+	for (int32 LODIndex = 0; LODIndex < LODCount; ++LODIndex)
+	{
+		if (!LODs[LODIndex].bReduceMesh)
+		{
+			BaseLODIndex = LODIndex;
+		}
+	}
+	
+	// Pick which LOD will store our Mesh description (the other LODs will just refer to it)
+	// if there is a full resolution base LOD then we use that, otherwise just pick the first one
+	int32 MeshDescriptionLODIndex = BaseLODIndex < 0 ? 0 : BaseLODIndex;
+	check(MeshDescriptionLODIndex >= 0);
+
+	// build the full resolution mesh description, into MeshDescriptionLODIndex LOD
+	FMeshDescription* FullResMeshDescription = nullptr;
+	FStaticMeshSourceModel& FullResMeshDescriptionSourceModel = StaticMesh->GetSourceModel(MeshDescriptionLODIndex);
+	FullResMeshDescription = FullResMeshDescriptionSourceModel.GetCachedMeshDescription();
+
+	if ((FullResMeshDescription != nullptr) && (TransformerUnit.MeshData != nullptr))
+	{
+		TransformerUnit.MeshData->ConvertToMeshDescription(*FullResMeshDescription);
+	}
+
+	check(FullResMeshDescription != nullptr);
+
+	// fill out all of the LOD source models with mesh descriptions and LOD reduction settings
+	for (int32 LODIndex = 0; LODIndex < LODCount; ++LODIndex)
+	{
+		const MeshPartition::FLODSettings& LODSettings = LODs[LODIndex];
+		FStaticMeshSourceModel& LODSourceModel = StaticMesh->GetSourceModel(LODIndex);
+
+		FMeshBuildSettings& BuildSettings = LODSourceModel.BuildSettings;
+		BuildSettings.bRecomputeNormals = true;
+		BuildSettings.bRecomputeTangents = true;
+		BuildSettings.bGenerateLightmapUVs = false;
+		BuildSettings.DistanceFieldResolutionScale = 0.0f;
+		BuildSettings.MaxLumenMeshCards = 4;
+
+		if (LODIndex == MeshDescriptionLODIndex)
+		{
+			// this LOD has already been populated in the initial ConvertToMeshDescription
+			FMeshDescription* LODMeshDescription = LODSourceModel.GetCachedMeshDescription();
+			check(LODMeshDescription == FullResMeshDescription);
+		}
+
+		// if we're reducing the mesh, set up the settings
+		if (LODSettings.bReduceMesh)
+		{
+			LODSourceModel.ReductionSettings = LODSettings.ReductionSettings;
+		}
+		
+		// set up the reference to the full rez mesh description
+		LODSourceModel.ReductionSettings.BaseLODModel = MeshDescriptionLODIndex;
+	}
+
+	// Apply the LOD streaming settings to the static mesh
+	StaticMesh->SetLODStreaming(LODStreaming);
+	StaticMesh->GetNumStreamedLODs() = NumStreamedLODs;
+
+	// Force the mesh description to precalculate indexing data structures in this async task
+	// so that when this function is eventually called on the GT it will not require any significant computation.
+	FullResMeshDescription->BuildIndexers();
+
+	// Use the hash as the guid to ensure that given the same input mesh data, we always hit the cache for the static mesh/nanite build.
+	FullResMeshDescriptionSourceModel.CommitMeshDescription(/* bUseHashAsGuid */ true);
+}
+
+void FStaticMeshTransformer::FinalizeStaticMeshes(MeshPartition::FTransformerContext& InTransformerContext, UE::Tasks::TTask<TArray<UStaticMesh*>>& InCreateStaticMeshesTask) const
+{
+	check(IsInGameThread());
+		   
+	TArray<UStaticMesh*>& StaticMeshes = InCreateStaticMeshesTask.GetResult();
+
+	for (int32 Index = 0; Index < InTransformerContext.TransformerUnits.Num(); ++Index)
+	{
+		const MeshPartition::FTransformerUnit& TransformerUnit = InTransformerContext.TransformerUnits[Index];
+
+		if (!ensure(TransformerUnit.Section != nullptr) || !ensure(Index < StaticMeshes.Num()))
+		{
+			continue;
+		}
+
+		UStaticMesh* StaticMesh = StaticMeshes[Index];
+
+		if (StaticMesh != nullptr)
+		{
+			StaticMesh->ClearInternalFlags(EInternalObjectFlags::Async);
+			ForEachObjectWithOuter(StaticMesh, [](UObject* Object)
+			{
+			   Object->ClearInternalFlags(EInternalObjectFlags::Async);
+			});
+		}
+				
+		FMeshDescription* MeshDescription = (StaticMesh != nullptr) ? StaticMesh->GetMeshDescription(0) : nullptr;
+				
+		if (!InTransformerContext.bWasCancelled && (MeshDescription != nullptr) && (MeshDescription->Vertices().Num() != 0))
+		{
+			MeshPartition::EditorUtils::FFinalizeStaticMeshParams Params;
+
+			Params.StaticMesh = StaticMesh;
+			Params.CollisionProfile = CollisionProfile;
+			Params.NumLODs = LODs.Num();
+			Params.bCanEverAffectNavigation = bCanEverAffectNavigation;
+			Params.bUseNanite = bUseNanite;
+			Params.bSetupSections = true;
+
+			if (MeshPartition::APreviewSection* PreviewSection = Cast<MeshPartition::APreviewSection>(TransformerUnit.Section))
+			{
+				Params.Material = PreviewSection->GetMaterialInstance();
+				MeshPartition::EditorUtils::FinalizeStaticMesh(Params);
+
+				PreviewSection->AddMesh(StaticMesh, CollisionProfile.Name, bCanEverAffectNavigation);
+			}
+			else if (MeshPartition::ACompiledSection* CompiledSection = Cast<MeshPartition::ACompiledSection>(TransformerUnit.Section))
+			{
+				Params.Material = CompiledSection->GetMaterialInstance();
+				MeshPartition::EditorUtils::FinalizeStaticMesh(Params);
+
+				CompiledSection->AddStaticMesh(StaticMesh, CollisionProfile.Name, bCanEverAffectNavigation);
+			}
+		}
+	}
+}
+
+} // namespace UE::MeshPartition
diff --git a/Engine/Plugins/Experimental/MeshPartition/Source/MeshPartitionEditor/Private/MeshPartitionSubsectionTransformer.cpp b/Engine/Plugins/Experimental/MeshPartition/Source/MeshPartitionEditor/Private/MeshPartitionSubsectionTransformer.cpp
new file mode 100644
--- /dev/null
+++ b/Engine/Plugins/Experimental/MeshPartition/Source/MeshPartitionEditor/Private/MeshPartitionSubsectionTransformer.cpp
@@ -0,0 +1,60 @@
+// Copyright Epic Games, Inc. All Rights Reserved.
+
+#include "MeshPartitionSubsectionTransformer.h"
+
+#include "MeshPartitionDependencyInterface.h"
+#include "MeshPartitionMeshBuilder.h"
+
+namespace UE::MeshPartition
+{
+
+namespace MeshPartitionSubsectionTransformerLocals
+{
+	void AddSubSections(const MeshPartition::FTransformerUnit& InTransformerUnit, AActor* InSection, const uint32 InSubSectionSize, TArray<MeshPartition::FTransformerUnit>& OutResults)
+	{
+		const MeshPartition::FMeshData& MeshData = *InTransformerUnit.MeshData;
+
+		if (InSubSectionSize == 0)
+		{
+			return;
+		}
+		
+		TArray<MeshPartition::FMeshData> SubSections = UE::MeshPartition::BuildHelpers::BuildSections(MeshData, InSubSectionSize);
+
+		OutResults.Reserve(OutResults.Num() + SubSections.Num());  
+
+		for (MeshPartition::FMeshData& SubSection : SubSections)
+		{
+			MeshPartition::FTransformerUnit& NewEntry = OutResults.Emplace_GetRef();
+			NewEntry.Section = InSection;
+			NewEntry.MeshData = MakeShared<const MeshPartition::FMeshData>(MoveTemp(SubSection));
+			NewEntry.bShouldRecomputeNormals = InTransformerUnit.bShouldRecomputeNormals;
+			NewEntry.bShouldRecomputeTangents = InTransformerUnit.bShouldRecomputeTangents;
+		}
+	}
+	
+} // namespace MeshPartitionSubsectionTransformerLocals
+
+bool FSubsectionTransformer::Execute(MeshPartition::FTransformerContext& InTransformerContext) const
+{
+	TArray<MeshPartition::FTransformerUnit> Results;
+
+	for (const MeshPartition::FTransformerUnit& TransformerUnit : InTransformerContext.TransformerUnits)
+	{
+		MeshPartitionSubsectionTransformerLocals::AddSubSections(TransformerUnit, TransformerUnit.Section, SubSectionSize, Results);
+	}
+
+	if (!Results.IsEmpty())
+	{
+		InTransformerContext.TransformerUnits = MoveTemp(Results);
+	}
+
+	return true;
+}
+
+void FSubsectionTransformer::GatherDependencies(MeshPartition::IDependencyInterface& InDependencies) const
+{
+	InDependencies += SubSectionSize;
+}
+
+} // namespace UE::MeshPartition
