[Mesh Partition]

diff --git a/Engine/Plugins/Experimental/MeshPartition/Source/MeshPartitionEditor/Private/MeshPartitionWorldUpdater.cpp b/Engine/Plugins/Experimental/MeshPartition/Source/MeshPartitionEditor/Private/MeshPartitionWorldUpdater.cpp
new file mode 100644
--- /dev/null
+++ b/Engine/Plugins/Experimental/MeshPartition/Source/MeshPartitionEditor/Private/MeshPartitionWorldUpdater.cpp
@@ -0,0 +1,523 @@
+// Copyright Epic Games, Inc. All Rights Reserved.
+
+#include "MeshPartitionWorldUpdater.h"
+#include "EngineUtils.h"
+#include "MeshPartitionEditorModule.h" // LogMegaMeshEditor
+#include "MeshPartitionDependencyContext.h"
+#include "MeshPartitionEditorSubsystem.h"
+#include "AssetRegistry/AssetRegistryModule.h"
+#include "Engine/StaticMesh.h"
+#include "Editor/EditorEngine.h"
+#include "MeshPartitionEditorUtils.h"
+
+#define LOCTEXT_NAMESPACE "FMeshPartitionWorldUpdater"
+
+extern UNREALED_API UEditorEngine* GEditor;
+
+namespace UE::MeshPartition
+{
+	FMeshPartitionWorldUpdater::FMeshPartitionWorldUpdater(UWorld* InWorld, EUpdateMode InUpdateMode)
+	{
+		// create a unique build key to identify placeholders created
+		UpdateMode = InUpdateMode;
+		BuildKey = FGuid::NewGuid();
+
+		// we will need the asset registry to check the up-to-date status of compiled sections (we will force completion of any relevant pending scans before checking)
+		FAssetRegistryModule& AssetRegistryModule = FModuleManager::LoadModuleChecked<FAssetRegistryModule>("AssetRegistry");
+		IAssetRegistry* AssetRegistry = &AssetRegistryModule.Get();
+		check(AssetRegistry);
+
+		MeshPartition::FModifierDescriptorCache* DescriptorCache = UMeshPartitionEditorSubsystem::GetDescriptorCache();
+		check(DescriptorCache);
+
+		// in editor worlds, update the descriptor cache immediately
+		if (InWorld->IsEditorWorld())
+		{
+			DescriptorCache->UpdateCacheForAllMegaMeshesInWorld(InWorld);
+		}
+		else
+		{
+			UE_LOG(LogMegaMeshEditor, Verbose, TEXT("Constructing FMeshPartitionWorldUpdater on a non-editor world - this relies on previously cached descriptor data"));
+		}
+
+		// create an FMeshPartitionUpdater for each of the mesh partitions in the editor world
+		for (AMeshPartition* MeshPartition : TActorRange<AMeshPartition>(InWorld))
+		{
+			if (MeshPartition)
+			{
+				const FGuid& MeshPartitionGUID = MeshPartition->GetActorGuid();
+
+				const UMeshPartitionDefinition* MPDefinition = MeshPartition->GetMeshPartitionDefinition();
+				if (MPDefinition == nullptr)
+				{
+					UE_LOG(LogMegaMeshEditor, Verbose, TEXT("No Mesh Partition Definition provided for %s (%.6s), using default settings."), *MeshPartition->GetActorNameOrLabel(), *MeshPartitionGUID.ToString());
+					MPDefinition = UMeshPartitionDefinition::StaticClass()->GetDefaultObject<UMeshPartitionDefinition>();
+				}
+				check(MPDefinition);
+
+				TSharedPtr<MeshPartition::FCachedDescriptors> MeshPartitionAllCachedDescriptors = DescriptorCache->GetCachedDescriptors(MeshPartition);
+				if (!MeshPartitionAllCachedDescriptors.IsValid())
+				{
+					UE_LOG(LogMegaMeshEditor, Verbose, TEXT("Ran without descriptors cached for MeshPartition '%s', attempting to update the cache in the current world"), *MeshPartition->GetActorNameOrLabel());
+					DescriptorCache->UpdateCacheForAllMegaMeshesInWorld(InWorld);
+					MeshPartitionAllCachedDescriptors = DescriptorCache->GetCachedDescriptors(MeshPartition);
+					check(MeshPartitionAllCachedDescriptors.IsValid());
+				}
+				check(MeshPartitionAllCachedDescriptors->GetMegaMeshGUID() == MeshPartition->GetActorGuid());
+
+				TArray<FName> BuildVariantNames;
+				switch (UpdateMode)
+				{
+					case EUpdateMode::ForPIE:
+						{
+							// in PIE, process variants for the preview platform only
+							FName PreviewPlatformName;
+							ITargetPlatform* PreviewTargetPlatform = UE::MeshPartition::EditorUtils::GetPreviewPlatform(PreviewPlatformName);
+							BuildVariantNames = MPDefinition->GetCompiledSectionBuildVariantNamesForCurrentEditorPreview(PreviewPlatformName, PreviewTargetPlatform);
+						}
+						break;
+					case EUpdateMode::ForCompile:
+						// when compiling, process all build variants
+						BuildVariantNames = MPDefinition->GetAllCompiledSectionBuildVariantNames();
+						break;
+					case EUpdateMode::ForCook:
+						{
+							// when cooking, prepare build variants for all target platforms
+							TArray<ITargetPlatform*> TargetPlatforms = GetTargetPlatformManager()->GetActiveTargetPlatforms();
+							BuildVariantNames = MPDefinition->GetCompiledSectionBuildVariantNamesForPlatforms(TargetPlatforms);
+						}
+						break;
+				}
+				check(!BuildVariantNames.IsEmpty());
+
+				// add an entry for this mesh partition
+				FMeshPartitionUpdater& MeshPartitionUpdater = MeshPartitionUpdaters.Add(
+					MeshPartitionGUID,
+					FMeshPartitionUpdater(BuildKey, *MPDefinition, MeshPartitionAllCachedDescriptors, TConstArrayView<FName>(BuildVariantNames), *AssetRegistry));
+
+				// in PIE, we create placeholders for any missing bases.  In cook, we just warn that they are missing.
+				if (UpdateMode == EUpdateMode::ForCook)
+				{
+					MeshPartitionUpdater.ForAllBuildVariants([](const FBuildVariantUpdater& BuildVariantUpdater)
+						{
+							if (BuildVariantUpdater.GetBaseModifiersNeedingCreation().Num() > 0)
+							{
+								UE_LOG(LogMegaMeshEditor, Warning, TEXT("Mesh Partition '%s' BuildVariant:%s is missing compiled sections for %d base modifiers.  Cooked builds will not show these missing regions until the compiled sections are built."),
+									*BuildVariantUpdater.GetMeshPartitionPath().ToString(),
+									*BuildVariantUpdater.GetBuildVariantName().ToString(),
+									BuildVariantUpdater.GetBaseModifiersNeedingCreation().Num());
+							}
+							return true;
+						});
+				}
+			}
+		}
+	}
+
+	FMeshPartitionUpdater::FMeshPartitionUpdater(const FGuid& InBuildKey, const UMeshPartitionDefinition& InMeshPartitionDefinition, TSharedPtr<MeshPartition::FCachedDescriptors> InMeshPartitionAllCachedDescriptors, TConstArrayView<FName> InTargetBuildVariantNames, IAssetRegistry& InAssetRegistry)
+		: AllCachedDescriptors(InMeshPartitionAllCachedDescriptors)
+	{
+		check(AllCachedDescriptors.IsValid());
+
+		// create an FBuildVariantUpdater for each of the target build variants in the mesh partition
+		for (const FName& TargetBuildVariantName : InTargetBuildVariantNames)
+		{
+			const MeshPartition::FModifierGrouping* TargetGrouping = AllCachedDescriptors->GetCachedBuildVariant(TargetBuildVariantName);
+			if (!TargetGrouping)
+			{
+				UE_LOG(LogMegaMeshEditor, Error, TEXT("Could not find build variant '%s' in the cached descriptors, this may result in incorrect behaviors in MeshPartition"), *TargetBuildVariantName.ToString());
+			}
+			else
+			{
+				check(TargetGrouping->GetBuildVariantName() == TargetBuildVariantName);
+				BuildVariantUpdaters.Emplace(TargetBuildVariantName, FBuildVariantUpdater(InBuildKey, InMeshPartitionDefinition, *TargetGrouping, AllCachedDescriptors, InAssetRegistry));
+			}
+		}
+
+		// grab the compiled section descriptor array, initialize the corresponding status array with default state, and initialize the lookup table for compiled sections
+		CompiledSections = AllCachedDescriptors->GetCompiledSections();
+		CompiledSectionStatus.Init(MeshPartition::FCompiledSectionStatus(), CompiledSections.Num());
+		for (int32 CompiledSectionIndex = 0; CompiledSectionIndex < CompiledSections.Num(); CompiledSectionIndex++)
+		{
+			const FCompiledSectionDescriptor& CompiledSectionDesc = CompiledSections[CompiledSectionIndex];
+			ActorGuidToCompiledSectionIndex.Add(CompiledSectionDesc.ActorDescGuid, CompiledSectionIndex);
+		}
+
+		// let's build the correspondence between the target groups (what we want)
+		// and the compiled sections (what we have), and check the up to date status.
+		// 
+		// For a compiled section to be an up-to-date match for a target group, we require:
+		// 1) same mesh partition, (GUID, path, definition, build variant)
+		// 2) same set of modifiers (modifier set hash, which is the hash of the modifier paths)
+		// 3) same BuildVariantHash - settings on the definition and build variant have not changed
+		// 4) package hash and class hashes are up to date - the content and code used to build the compiled section have not changed
+		// 
+		// we must also check for duplicates, and ensure we have at most one up to date compiled section per group
+		// 
+
+		// iterate the compiled sections, looking for any that match a target group on our build variants (and then update the compiled section status appropriately)
+		for (int32 CompiledSectionIndex = 0; CompiledSectionIndex < CompiledSections.Num(); CompiledSectionIndex++)
+		{
+			const FCompiledSectionDescriptor& CompiledSection = CompiledSections[CompiledSectionIndex];
+			MeshPartition::FCompiledSectionStatus& Status = CompiledSectionStatus[CompiledSectionIndex];
+
+			check(!Status.bChecked);
+			Status.bChecked = true;
+
+			// ignore placeholders -- they shouldn't be considered a properly built compiled section.
+			// we don't serialize placeholder status into the build info, but we can deduce it from an unset modifier set hash
+			const MeshPartition::FCompiledSectionBuildInfo& BuildInfo = CompiledSection.Info;
+			if (!BuildInfo.ModifierSetHash.IsValid())
+			{
+				Status.bIsPlaceholder = true;
+				continue;
+			}
+
+			// find the corresponding build variant updater
+			FBuildVariantUpdater* BuildVariantUpdater = BuildVariantUpdaters.Find(BuildInfo.BuildVariantName);
+			if (BuildVariantUpdater == nullptr)
+			{
+				Status.bNonTargetVariant = true;
+				continue;
+			}
+
+			// figure out which group this compiled section maps to, by looking up its modifier set hash
+			int32 GroupIndex = BuildVariantUpdater->GetTargetGroupIndexFromModifierSetHash(BuildInfo.ModifierSetHash);
+			if (GroupIndex < 0)
+			{
+				Status.bMismatchedModifiers = true;
+				continue;
+			}
+
+			Status.bPackageHashMatches = UE::MeshPartition::EditorUtils::PackageHashIsUpToDate(BuildInfo, InAssetRegistry);
+			Status.bClassHashMatches = UE::MeshPartition::EditorUtils::ClassHashIsUpToDate(BuildInfo);
+
+			const MeshPartition::FModifierGroup& ModifierGroup = BuildVariantUpdater->GetTargetModifierGroups()[GroupIndex];
+
+			// double check that the modifier paths match exactly, because even 128 bit hashes might collide
+			check(ModifierGroup.BaseDescs().Num() == BuildInfo.BaseModifierPaths.Num());
+			for (int32 BaseModifierIndex = 0; BaseModifierIndex < ModifierGroup.BaseDescs().Num(); BaseModifierIndex++)
+			{
+				check(ModifierGroup.BaseDescs()[BaseModifierIndex].ModifierPath == BuildInfo.BaseModifierPaths[BaseModifierIndex]);
+			}
+
+			// TODO: this shouldn't be necessary -- ideally we should ensure ActorDesc is up to date before caching the descriptors...
+			// (currently after deleting a compiled section, the ActorDescs stick around for a while...)
+			// for now we double check if compiled sections were deleted
+			bool bStillExists = InAssetRegistry.GetAssetByObjectPath(CompiledSection.ActorPath).IsValid();
+			if (!bStillExists)
+			{
+				Status.bCannotFindFile = true;
+				continue;
+			}
+
+			const bool bUpToDate = bStillExists && Status.bPackageHashMatches && Status.bClassHashMatches;
+			if (bUpToDate)
+			{
+				// check if there is a duplicate compiled section for the same group
+				int32& BestCompiledSectionIndex = BuildVariantUpdater->BestReusableCompiledSectionIndexForGroup[GroupIndex];
+				if (BestCompiledSectionIndex < 0)
+				{
+					// no previous best compiled section, use this one
+					BestCompiledSectionIndex = CompiledSectionIndex;
+					Status.bIsDuplicate = false;
+				}
+				else
+				{
+					// there is a previous best -- determine which one to use
+					MeshPartition::FCompiledSectionStatus& PrevBestCompiledSectionStatus = CompiledSectionStatus[BestCompiledSectionIndex];
+					check(PrevBestCompiledSectionStatus.bChecked && PrevBestCompiledSectionStatus.bReuse);
+
+					// we just use an arbitrary deterministic method to determine which one to keep
+					if (CompiledSections[BestCompiledSectionIndex].ActorDescGuid < CompiledSection.ActorDescGuid)
+					{
+						// use the previous best -- this one is a duplicate (and should not be reused)
+						Status.bIsDuplicate = true;
+					}
+					else
+					{
+						// use this one -- set previous best to be a duplicate
+						PrevBestCompiledSectionStatus.bIsDuplicate = true;
+						PrevBestCompiledSectionStatus.bReuse = false;
+						BestCompiledSectionIndex = CompiledSectionIndex;
+					}
+				}
+				Status.bReuse = bUpToDate && !Status.bIsDuplicate;
+			}
+		}
+
+		// Build the BaseModifiersNeedingCreation array
+		for (auto& Pair : BuildVariantUpdaters)
+		{
+			FBuildVariantUpdater& BuildVariantUpdater = Pair.Value;
+			for (int GroupIndex = 0; GroupIndex < BuildVariantUpdater.GetTargetModifierGroups().Num(); GroupIndex++)
+			{
+				bool bExistsReuse = true;
+
+				int32 CompiledSectionIndex = BuildVariantUpdater.BestReusableCompiledSectionIndexForGroup[GroupIndex];
+				if (CompiledSectionIndex < 0)
+				{
+					// this group does not have a reusable compiled section, and needs to be created.
+					// add all of the target group's base modifiers to the list of those needing to be created
+					const MeshPartition::FModifierGroup& ModifierGroup = BuildVariantUpdater.GetTargetModifierGroups()[GroupIndex];
+					for (const MeshPartition::FModifierDesc& ModifierDesc : ModifierGroup.BaseDescs())
+					{
+						BuildVariantUpdater.BaseModifiersNeedingCreation.Add(ModifierDesc.ModifierPath, ModifierDesc);
+					}
+				}
+				else
+				{
+					// sanity check
+					MeshPartition::FCompiledSectionStatus& Status = CompiledSectionStatus[CompiledSectionIndex];
+					check(Status.bChecked && Status.bReuse);
+				}
+			}
+		}
+	}
+
+	FBuildVariantUpdater::FBuildVariantUpdater(	const FGuid& InBuildKey, const UMeshPartitionDefinition& InMeshPartitionDefinition, const MeshPartition::FModifierGrouping& InTargetGrouping, TSharedPtr<MeshPartition::FCachedDescriptors> InAllMeshPartitionDescriptors, IAssetRegistry& InAssetRegistry)
+		: AllMeshPartitionDescriptors(InAllMeshPartitionDescriptors)
+		, TargetGrouping(InTargetGrouping)
+	{
+		using namespace EditorUtils;
+
+		check(&TargetGrouping != nullptr);
+		check(&InMeshPartitionDefinition != nullptr);
+		check(AllMeshPartitionDescriptors.IsValid());
+
+		// construct the BuildInfo
+		{
+			BuildInfo.BuildKey = InBuildKey;
+			BuildInfo.SetMegaMeshDefinition(&InMeshPartitionDefinition);
+			BuildInfo.MegaMeshGUID = InAllMeshPartitionDescriptors->GetMegaMeshGUID();
+			BuildInfo.BuildVariantName = InTargetGrouping.GetBuildVariantName();
+			{	// hash the build variant
+				FDependencyHash DepdendencyHash;
+				InMeshPartitionDefinition.GatherDependencies(DepdendencyHash, BuildInfo.BuildVariantName);
+				BuildInfo.BuildVariantHash = DepdendencyHash.GetDependentDataHash();
+			}
+			// Note: get the path from the descriptor ref, because it will be the original path (not the PIE temporary path)
+			BuildInfo.MegaMeshPath = InAllMeshPartitionDescriptors->GetMegaMeshPath();
+		}
+
+		// Note: this includes compiled sections that are not for our specific build variant
+		TConstArrayView<FCompiledSectionDescriptor> AllExistingCompiledSections = InAllMeshPartitionDescriptors->GetCompiledSections();
+
+		if (InAssetRegistry.IsGathering())
+		{
+			UE_LOG(LogMegaMeshEditor, Log, TEXT("AssetRegistry is not complete, forcing synchronous scan of the MegaMesh related assets"));
+
+			// force the AssetRegistry to complete scanning on the packages we care about
+			TArray<FName> PackageNamesToQuery;
+			for (const FCompiledSectionDescriptor& ExistingCompiledSection : AllExistingCompiledSections)
+			{
+				PackageNamesToQuery.Add(ExistingCompiledSection.ActorPath.GetLongPackageFName());
+				PackageNamesToQuery.Append(ExistingCompiledSection.Info.PackageDependencies);
+			}
+
+			// convert package names to file paths
+			TArray<FString> PackageFilePathsToQuery;
+			for (FName PackageName : PackageNamesToQuery)
+			{
+				FString PackageFilename;
+				if (FPackageName::DoesPackageExist(PackageName.ToString(), &PackageFilename))
+				{
+					PackageFilePathsToQuery.Add(PackageFilename);
+				}
+			}
+
+			// ask asset registry to complete scanning the packages
+			InAssetRegistry.ScanFilesSynchronous(PackageFilePathsToQuery);
+		}
+
+		TConstArrayView<MeshPartition::FModifierGroup> TargetModifierGroups = InTargetGrouping.GetModifierGroups();
+
+		// init all groups to have no corresponding compiled section (we will determine the mapping later in the MeshPartition updater)
+		BestReusableCompiledSectionIndexForGroup.Init(-1, TargetModifierGroups.Num());
+
+		// build the ModifierSetHashToTargetGroupIndex
+		for (int GroupIndex = 0; GroupIndex < TargetModifierGroups.Num(); GroupIndex++)
+		{
+			const MeshPartition::FModifierGroup& ModifierGroup = TargetModifierGroups[GroupIndex];
+			FGuid ModifierSetHash = ModifierGroup.ComputeModifierSetHash();
+			ModifierSetHashToTargetGroupIndex.Add(ModifierSetHash, GroupIndex);
+		}
+	}
+
+
+	int32 FBuildVariantUpdater::GetReusableCompiledSectionForGroup(int32 GroupIndex, const FMeshPartitionUpdater& MeshPartitionUpdater) const
+	{
+		int32 CompiledSectionIndex = BestReusableCompiledSectionIndexForGroup[GroupIndex];
+		if (const FCompiledSectionStatus* Status = MeshPartitionUpdater.GetCompiledSectionStatus(CompiledSectionIndex))
+		{
+			ensure(Status->bChecked); // should have been classified by now
+			ensure(Status->bReuse);
+			// this group has a reusable compiled section
+			return CompiledSectionIndex;
+		}
+		return -1;
+	}
+
+	int32 FMeshPartitionWorldUpdater::CreatePlaceholderActors(UWorld* InTargetWorld)
+	{
+		check(UpdateMode == EUpdateMode::ForPIE);
+
+		int32 PlaceholdersCreated = 0;
+		for (TPair<FGuid, FMeshPartitionUpdater>& Pair : MeshPartitionUpdaters)
+		{
+			FMeshPartitionUpdater& MeshPartitionUpdater = Pair.Value;
+			MeshPartitionUpdater.ForAllBuildVariants([InTargetWorld, &MeshPartitionUpdater, &PlaceholdersCreated](const FBuildVariantUpdater& BuildVariantUpdater)
+			{
+				TConstArrayView<MeshPartition::FModifierGroup> ModifierGroups = BuildVariantUpdater.GetTargetGrouping().GetModifierGroups();
+				for (int32 GroupIndex = 0; GroupIndex < ModifierGroups.Num(); GroupIndex++)
+				{
+					int32 CompiledSectionIndex = BuildVariantUpdater.GetReusableCompiledSectionForGroup(GroupIndex, MeshPartitionUpdater);
+					if (CompiledSectionIndex >= 0)
+					{
+						continue; // no need to create a placeholder compiled section if there this group already has a reuseable one
+					}
+
+					// insert a placeholder for this modifier group!
+					const MeshPartition::FModifierGroup& ModifierGroup = ModifierGroups[GroupIndex];
+					const MeshPartition::UMeshPartitionDefinition* Definition = BuildVariantUpdater.GetBuildInfo().GetMegaMeshDefinition();
+					const MeshPartition::FCompiledSectionBuildVariant& BuildVariant = Definition->GetCompiledSectionBuildVariantByName(BuildVariantUpdater.GetBuildVariantName());
+
+					// Build the bases list from the modifier paths
+					FCompiledSectionBuildInfo BuildInfo = BuildVariantUpdater.GetBuildInfo();
+					BuildInfo.BaseModifierPaths.Empty();
+					Algo::Transform(ModifierGroup.BaseDescs(), BuildInfo.BaseModifierPaths, [](const MeshPartition::FModifierDesc& BaseDescriptor)
+						{
+							// the path is taken from the editor world, so it should not need any massaging (unlike PIE world paths)
+							return BaseDescriptor.ModifierPath;
+						});
+
+					// Note: We don't fill out ModifiersSetHash because we don't have the full set of non-base modifiers
+					// (The placeholder builder will find them when it loads everything)
+					FString BaseName = TEXT("CompiledSection_") + BuildVariant.Name.ToString() + TEXT("_PIE");
+					FName UniqueName = MakeUniqueObjectName(InTargetWorld, UStaticMesh::StaticClass(), FName(BaseName), EUniqueObjectNameOptions::UniversallyUnique);
+
+					FActorSpawnParameters Params;
+					Params.Name = UniqueName;
+					// Params.Owner = MeshPartitionActor; // TODO: we could parent to the main actor, but we would have to search for it first
+					Params.SpawnCollisionHandlingOverride = ESpawnActorCollisionHandlingMethod::AlwaysSpawn;
+					Params.bNoFail = true;
+					Params.ObjectFlags = RF_NoFlags;		// non transational (does not interact with undo/redo)
+
+					MeshPartition::ACompiledSection* Placeholder = InTargetWorld->SpawnActor<MeshPartition::ACompiledSection>(MeshPartition::ACompiledSection::StaticClass(), FTransform::Identity, Params);
+					if (Placeholder->GetRootComponent())
+					{
+						Placeholder->GetRootComponent()->SetMobility(EComponentMobility::Static);
+					}
+					Placeholder->SetBuildInfo(BuildInfo);
+					Placeholder->SetIsPlaceholder(true);
+					Placeholder->SetActorLabel(BaseName);
+
+					// build the estimated streaming bounds by unioning each base descriptor's bounds
+					FBox EstimatedStreamingBounds(ForceInit);
+					{
+						for (const FModifierDesc& Desc : ModifierGroup.BaseDescs())
+						{
+							EstimatedStreamingBounds += Desc.Bounds;
+						}
+						// apply some wiggle room for non-base modifiers -- this can be wildly off though..
+						EstimatedStreamingBounds = EstimatedStreamingBounds.ExpandBy(EstimatedStreamingBounds.GetExtent() * 0.1);
+					}
+					Placeholder->SetPlaceholderStreamingBounds(EstimatedStreamingBounds);
+
+					UE_LOG(LogMegaMeshEditor, Verbose, TEXT("    ADDED Placeholder %.6s (%p %s) (for %d bases)"), *Placeholder->GetActorGuid().ToString(), Placeholder, *Placeholder->GetActorNameOrLabel(), Placeholder->GetBuildInfo().BaseModifierPaths.Num());
+					PlaceholdersCreated++;
+					MeshPartitionUpdater.AddedPlaceholderCompiledSections++;
+				}
+				return true; // continue iterating build variants
+			});
+		}
+		PlaceholdersAddedToWorld = InTargetWorld;
+		return PlaceholdersCreated;
+	}
+
+	int32 FMeshPartitionWorldUpdater::RemovePlaceholderActors()
+	{
+		check(UpdateMode == EUpdateMode::ForPIE);
+
+		int32 PlaceholdersRemoved = 0;
+		if (UWorld* World = PlaceholdersAddedToWorld.Get())
+		{
+			for (TActorIterator<MeshPartition::ACompiledSection> It(World); It; ++It)
+			{
+				MeshPartition::ACompiledSection* CompiledSection = *It;
+				if (CompiledSection)
+				{
+					if (CompiledSection->IsPlaceholder())
+					{
+						CompiledSection->Destroy();
+						PlaceholdersRemoved++;
+					}
+				}
+			}
+		}
+		PlaceholdersAddedToWorld = nullptr;
+		return PlaceholdersRemoved;
+	}
+
+	void FMeshPartitionUpdater::ForAllBuildVariants(TFunctionRef<bool(const MeshPartition::FBuildVariantUpdater&)> InFunc) const
+	{
+		for (const TPair<FName, FBuildVariantUpdater>& Pair : BuildVariantUpdaters)
+		{
+			const FBuildVariantUpdater& BuildVariantUpdater = Pair.Value;
+			if (!InFunc(BuildVariantUpdater))
+			{
+				return;
+			}
+		}
+	}
+
+	const FCompiledSectionStatus* FMeshPartitionUpdater::GetCompiledSectionStatus(int32 CompiledSectionIndex) const
+	{
+		if (CompiledSectionStatus.IsValidIndex(CompiledSectionIndex))
+		{
+			return &CompiledSectionStatus[CompiledSectionIndex];
+		}
+		return nullptr;
+	}
+
+	const FCompiledSectionStatus* FMeshPartitionUpdater::GetCompiledSectionStatus(const ACompiledSection* CompiledSection) const
+	{
+		int32 CompiledSectionIndex = ActorGuidToCompiledSectionIndex.FindRef(CompiledSection->GetActorGuid(), -1);
+		if (CompiledSectionIndex >= 0)
+		{
+			return &CompiledSectionStatus[CompiledSectionIndex];
+		}
+		return nullptr;
+	}
+
+	const FMeshPartitionUpdater* FMeshPartitionWorldUpdater::GetMeshPartitionUpdaterFor(const ACompiledSection* CompiledSection) const
+	{
+		check(CompiledSection);
+		return MeshPartitionUpdaters.Find(CompiledSection->GetBuildInfo().MegaMeshGUID);
+	}
+
+	const FCompiledSectionStatus* FMeshPartitionWorldUpdater::GetCompiledSectionStatus(const ACompiledSection* CompiledSection) const
+	{
+		if (const FMeshPartitionUpdater* MeshPartitionUpdater = GetMeshPartitionUpdaterFor(CompiledSection))
+		{
+			return MeshPartitionUpdater->GetCompiledSectionStatus(CompiledSection);
+		}
+		return nullptr;
+	}
+
+	void FMeshPartitionWorldUpdater::ReportStats() const
+	{
+		for (const TPair<FGuid, FMeshPartitionUpdater>& Pair : MeshPartitionUpdaters)
+		{
+			const FMeshPartitionUpdater& MeshPartitionUpdater = Pair.Value;
+
+			UE_LOG(LogMegaMeshEditor, Verbose, TEXT("  -- MeshPartition '%s' : %d out of date, %d non-target build variant, %d reused, %d placeholders created"),
+				*MeshPartitionUpdater.AllCachedDescriptors->GetMegaMeshPath().ToString(),
+				MeshPartitionUpdater.RemovedOutOfDate,
+				MeshPartitionUpdater.RemovedNonTargetBuildVariant,
+				MeshPartitionUpdater.ReusedCompiledSections,
+				MeshPartitionUpdater.AddedPlaceholderCompiledSections);
+		}
+	}
+} // namespace UE::MeshPartition
+#undef LOCTEXT_NAMESPACE
\ No newline at end of file
diff --git a/Engine/Plugins/Experimental/MeshPartition/Source/MeshPartitionEditor/Private/Test/MeshPartitionPatchTests.cpp b/Engine/Plugins/Experimental/MeshPartition/Source/MeshPartitionEditor/Private/Test/MeshPartitionPatchTests.cpp
new file mode 100644
--- /dev/null
+++ b/Engine/Plugins/Experimental/MeshPartition/Source/MeshPartitionEditor/Private/Test/MeshPartitionPatchTests.cpp
@@ -0,0 +1,160 @@
+// Copyright Epic Games, Inc. All Rights Reserved.
+
+#if WITH_DEV_AUTOMATION_TESTS
+#if WITH_EDITOR
+
+#include "CQTest.h"
+#include "Test/MeshPartitionTestUtils.h"
+#include "EditorWorldUtils.h"
+#include "MeshPartition.h"
+#include "MeshPartitionEditorComponent.h"
+#include "MeshPartitionPreviewSection.h"
+#include "Engine/World.h"
+#include "Tests/AutomationCommon.h"
+#include "Tests/AutomationEditorCommon.h"
+
+using namespace UE::MeshPartition;
+
+/**
+ * Tests for MegaMesh Patch modifier
+ */
+TEST_CLASS_WITH_FLAGS(Patch, "MeshPartition.Modifier", EAutomationTestFlags::EditorContext | EAutomationTestFlags::EngineFilter)
+{
+private:
+	// Test environment and core objects
+	TUniquePtr<FScopedEditorWorld> ScopedEditorWorld;
+	// CVars
+	TSharedPtr<FScopedTestEnvironment> TestEnvironment;
+
+	// Core test objects
+	UWorld* World = nullptr;
+	AMeshPartition* MegaMesh = nullptr;
+	UMeshPartitionEditorComponent* MegaMeshEditorComponent = nullptr;
+
+	// Test constants
+	static constexpr int32 MeshWidth = 20000;
+	static constexpr int32 MeshHeight = 20000;
+	static constexpr int32 MeshWidthVertexCount = 100;
+	static constexpr int32 MeshHeightVertexCount = 100;
+	static constexpr int32 ExpectedPreviewSectionCount = 1;
+	static inline const FString LevelPath = TEXT("/MeshPartition/UnitTests/MeshPartitionUnitTests");
+
+public:
+	BEFORE_EACH()
+	{
+		// Configure minimal editor world
+		FWorldInitializationValues InitializationValues;
+		InitializationValues.RequiresHitProxies(false);
+		InitializationValues.ShouldSimulatePhysics(false);
+		InitializationValues.EnableTraceCollision(false);
+		InitializationValues.CreateNavigation(false);
+		InitializationValues.CreateAISystem(false);
+		InitializationValues.AllowAudioPlayback(false);
+		InitializationValues.CreatePhysicsScene(true);
+
+		// Spawn world
+		ScopedEditorWorld = FAutomationEditorCommonUtils::CreateScopedEditorWorld(LevelPath, InitializationValues);
+		ASSERT_THAT(IsNotNull(ScopedEditorWorld.Get(), TEXT("Failed to create scoped editor world")));
+
+		World = ScopedEditorWorld->GetWorld();
+		ASSERT_THAT(IsNotNull(World, TEXT("Failed to create valid test world")));
+
+		// Need to disable DDC read/write, so that generated preview sections aren't written to DDC
+		TestEnvironment = FScopedTestEnvironment::Get();
+		ASSERT_THAT(IsNotNull(TestEnvironment.Get(), TEXT("Failed to get scoped test environment")));
+
+		TestEnvironment->SetConsoleVariableValue(TEXT("MegaMesh.Preview.EnableDDCWrite"), TEXT("0"));
+		TestEnvironment->SetConsoleVariableValue(TEXT("MegaMesh.Preview.EnableDDCRead"), TEXT("0"));
+	}
+
+	/**
+	 * Verifies that when two patch modifiers with overlapping areas are applied, 
+	 * the modifier with higher priority takes precedence over the lower one.
+	 */
+	TEST_METHOD(Verify_PatchModifier_Priority_System_Applies_Higher_Priority_First)
+	{
+		TestCommandBuilder
+			.Do(TEXT("Create MegaMesh with test dimensions"), [this]()
+			{
+				// Create MegaMesh
+				MegaMesh = TestUtils::CreateTestMesh(
+					World, MeshWidth, MeshHeight, MeshWidthVertexCount, MeshHeightVertexCount);
+				ASSERT_THAT(IsNotNull(MegaMesh, TEXT("Failed to create MegaMesh")));
+
+				MegaMeshEditorComponent = Cast<UMeshPartitionEditorComponent>(MegaMesh->GetMeshPartitionComponent());
+				ASSERT_THAT(IsNotNull(MegaMeshEditorComponent,
+					TEXT("Failed to get MegaMeshEditorComponent from MegaMesh")));
+			})
+			.Then(TEXT("Create overlapping patch modifiers with different priorities"), [this]()
+			{
+				// Create high priority patch modifier
+				{
+					const FVector ModifierWorldLocation = FVector(0.0, 2000.0, 1500.0);
+					constexpr double HighPriority = 1.0;
+					AActor* HighPriorityPatchActor = World->SpawnActor<AActor>(AActor::StaticClass(),
+						FTransform::Identity);
+					ASSERT_THAT(IsNotNull(HighPriorityPatchActor, TEXT("Failed to spawn high priority MegaMeshModifier")));
+
+					TestUtils::CreatePatchModifier(HighPriorityPatchActor, MegaMesh, ModifierWorldLocation, HighPriority);
+				}
+				// Create low priority patch modifier
+				{
+					const FVector ModifierWorldLocation = FVector(0.0, 1000.0, 2500.0);
+					AActor* LowPriorityPatchActor = World->SpawnActor<AActor>(AActor::StaticClass(),
+						FTransform::Identity);
+					ASSERT_THAT(IsNotNull(LowPriorityPatchActor, TEXT("Failed to spawn low priority MegaMeshModifier")));
+
+					TestUtils::CreatePatchModifier(LowPriorityPatchActor, MegaMesh, ModifierWorldLocation);
+				}
+
+				// Notify the editor component that modifiers have been assigned
+				MegaMeshEditorComponent->OnModifierAssigned();
+			})
+			.Until(TEXT("Wait for preview section build to complete"), [this]()
+			{
+				return !MegaMeshEditorComponent->IsAnyPreviewSectionBuildActive();
+			})
+			.Then(TEXT("Validate preview section matches expected mesh"), [this]()
+			{
+				const APreviewSection* PreviewSection = nullptr;
+				int32 PreviewSectionCount = 0;
+
+				// Count and capture preview sections
+				MegaMeshEditorComponent->ForAllPreviewSections([&](const APreviewSection* InPreviewSection)
+				{
+					PreviewSection = InPreviewSection;
+					++PreviewSectionCount;
+					return true;
+				});
+
+				// Validate preview section count
+				ASSERT_THAT(AreEqual(ExpectedPreviewSectionCount, PreviewSectionCount,
+					*FString::Printf(TEXT("PreviewSectionCount should be %d"), ExpectedPreviewSectionCount)));
+				ASSERT_THAT(IsNotNull(PreviewSection, TEXT("MegaMesh PreviewSection is null")));
+
+				// Validate preview mesh matches ground truth
+				const FString GroundTruthAsset = TEXT("Verify_PatchModifier_Priority_System_Applies_Higher_Priority_First_GroundTruth");
+				const FString GroundTruthPath = FString::Printf(
+					TEXT("/MeshPartition/UnitTests/PatchModifierTests/%s.%s"),
+					*GroundTruthAsset,
+					*GroundTruthAsset
+				);
+				FString ErrorMessage;
+				const bool bValid = TestUtils::ValidatePreviewMeshAgainstGroundTruth(
+					MegaMeshEditorComponent,
+					{ GroundTruthPath },
+					ErrorMessage);
+
+				ASSERT_THAT(IsTrue(bValid, *ErrorMessage));
+			});
+	}
+
+	AFTER_EACH()
+	{
+		ScopedEditorWorld.Reset();
+		CollectGarbage(GARBAGE_COLLECTION_KEEPFLAGS);
+	}
+};
+
+#endif // WITH_EDITOR
+#endif // WITH_DEV_AUTOMATION_TESTS
diff --git a/Engine/Plugins/Experimental/MeshPartition/Source/MeshPartitionEditor/Private/WorldPartitionMeshPartitionBuilder.cpp b/Engine/Plugins/Experimental/MeshPartition/Source/MeshPartitionEditor/Private/WorldPartitionMeshPartitionBuilder.cpp
new file mode 100644
--- /dev/null
+++ b/Engine/Plugins/Experimental/MeshPartition/Source/MeshPartitionEditor/Private/WorldPartitionMeshPartitionBuilder.cpp
@@ -0,0 +1,961 @@
+// Copyright Epic Games, Inc. All Rights Reserved.
+
+#include "WorldPartitionMeshPartitionBuilder.h"
+
+#include "AssetCompilingManager.h"
+#include "EditorBuildUtils.h"
+#include "Engine/StaticMesh.h"
+#include "HAL/PlatformFileManager.h"
+#include "Internationalization/Regex.h"
+#include "ISourceControlModule.h"
+#include "Engine/LevelStreaming.h"
+#include "MeshPartition.h"
+#include "MeshPartitionDefinition.h"
+#include "MeshPartitionCompiledSection.h"
+#include "MeshPartitionEditorModule.h"
+#include "MeshPartitionEditorComponent.h"
+#include "MeshPartitionModifierComponent.h"
+#include "MeshPartitionModifierDescriptors.h"
+#include "MeshPartitionModifierTaskGraph.h"
+#include "MeshPartitionEditorSubsystem.h"
+#include "Misc/App.h"
+#include "Misc/FileHelper.h"
+#include "Misc/ScopedSlowTask.h"
+#include "ObjectTools.h"
+#include "SourceControlHelpers.h"
+#include "UObject/SavePackage.h"
+#include "WorldPartition/WorldPartition.h"
+#include "WorldPartition/WorldPartitionHelpers.h"
+#include "EngineUtils.h"
+#include "Misc/MessageDialog.h"
+#include "MeshPartitionMeshBuilder.h"
+#include "Algo/ForEach.h"
+#include "MeshPartitionActorDescUtils.h"
+#include "MeshPartitionModifierUtils.h"
+#include "MeshPartitionCompiledSectionActorDesc.h"
+#include "AssetRegistry/AssetRegistryModule.h"
+#include "Misc/Optional.h"
+#include "UObject/UObjectIterator.h"
+#include "UObject/Class.h"
+#include "Commandlets/Commandlet.h"
+#include "Kismet2/BlueprintEditorUtils.h"
+#include "MeshPartitionEditorUtils.h"
+#include "MeshPartitionWorldPartitionHelpers.h"
+#include "MeshPartitionDataLayerContainer.h"
+#include "MeshPartitionPlatformCellTransformer.h"
+#include "WorldPartition/DataLayer/WorldDataLayers.h"
+#include "WorldPartition/DataLayer/DataLayerAsset.h"
+#include "WorldPartition/DataLayer/DataLayerManager.h"
+#include "WorldPartition/DataLayer/DataLayerInstance.h"
+
+#define LOCTEXT_NAMESPACE "MegaMeshBuilder"
+
+static FAutoConsoleVariable CVarMegaMeshBuildCommandWaitForDebugger(
+	TEXT("megamesh.Build.WaitForDebugger"),
+	0,
+	TEXT("Enables wait for debugger when starting a Mesh Partition build process"));
+
+namespace UE::MeshPartition
+{
+
+namespace WorldPartitionMegaMeshBuilderLocals
+{
+	int32 GetTotalNumberOfSectionsToBuild(MeshPartition::FModifierDescriptorCache* InDescriptorCache)
+	{
+		int32 TotalNumberOfSectionsToBuild = 0;
+
+		if (!ensure(InDescriptorCache != nullptr))
+		{
+			return 0;
+		}
+
+		InDescriptorCache->ForAllCachedDescriptors([&TotalNumberOfSectionsToBuild](const FSoftObjectPath& MeshPartitionPath, const TSharedPtr<MeshPartition::FCachedDescriptors>& CachedDescriptors)
+		{
+			CachedDescriptors->ForAllBuildVariants([&TotalNumberOfSectionsToBuild](const MeshPartition::FModifierGrouping& BuildVariant)
+			{
+				TotalNumberOfSectionsToBuild += BuildVariant.GetModifierGroups().Num();
+				return true;
+			});
+
+			return true;
+		});
+
+		return TotalNumberOfSectionsToBuild;
+	}
+
+	void LogBuildProgress(const int32 InTotalSectionNumber, const int32 InTotalNumberOfSectionsToBuild)
+	{
+		UE_LOG(LogMegaMeshEditor, Display, TEXT("[%i / %i] Building compiled sections..."), InTotalSectionNumber, InTotalNumberOfSectionsToBuild);
+	}
+}
+
+UWorldPartitionMeshPartitionBuilder::UWorldPartitionMeshPartitionBuilder(const FObjectInitializer& InObjectInitializer)
+	: Super(InObjectInitializer)
+{
+}
+
+bool UWorldPartitionMeshPartitionBuilder::PreRun(UWorld* World, FPackageSourceControlHelper& PackageHelper)
+{
+	TArray<FString> Tokens, Switches;
+	UCommandlet::ParseCommandLine(*GetBuilderArgs(), Tokens, Switches);
+	return UCommandlet::ApplyCommandLineSwitches(this, Switches);
+}
+
+bool UWorldPartitionMeshPartitionBuilder::RunInternal(UWorld* InWorld, const FCellInfo& InCellInfo, FPackageSourceControlHelper& InPackageHelper)
+{
+	using namespace Utils;
+
+	WorldPartition = InWorld->GetWorldPartition();
+	if (WorldPartition == nullptr)
+	{
+		UE_LOG(LogMegaMeshEditor, Error, TEXT("Failed to retrieve WorldPartition, cannot run Mesh Partition Build for World %s."), *InWorld->GetName());
+		return false;
+	}
+
+	MeshPartition::FModifierDescriptorCache* DescriptorCache = UMeshPartitionEditorSubsystem::GetDescriptorCache();
+	if (DescriptorCache == nullptr)
+	{
+		UE_LOG(LogMegaMeshEditor, Error, TEXT("Cannot access Mesh Partition DescriptorCache, cannot run Mesh Partition Build for World %s."), *InWorld->GetName());
+		return false;
+	}
+
+	TArray<FGuid> AllMegaMeshGuidsInWorld;
+	DescriptorCache->UpdateCacheForAllMegaMeshesInWorld(InWorld, &AllMegaMeshGuidsInWorld);
+
+	// The BuildKey is a unique identifier to tag all compiled sections built by this build execution (allowing us to remove old ones).
+	const FGuid BuildKey = FGuid::NewGuid();
+
+	UE_LOG(LogMegaMeshEditor, Display, TEXT("Build for World %s (%d Mesh Partitions, BuildKey %.8s)"), *InWorld->GetName(), AllMegaMeshGuidsInWorld.Num(), *BuildKey.ToString());
+
+	CompiledSectionActorDescGuidsToKeep.Empty();
+	TotalBuiltCompiledSections = 0;
+	TotalFailedCompiledSections = 0;
+	TotalReusedCompiledSections = 0;
+	TotalNumberOfSectionsToBuild = WorldPartitionMegaMeshBuilderLocals::GetTotalNumberOfSectionsToBuild(DescriptorCache);
+
+	UE_LOG(LogMegaMeshEditor, Display, TEXT("[%i / %i] Initializing compiled sections build..."), 0, TotalNumberOfSectionsToBuild);
+	
+	// For each MegaMesh in the world
+	for (const FGuid& MegaMeshGUID  : AllMegaMeshGuidsInWorld)
+	{
+		// Use FWorldPartitionReference to ensure the AMeshPartition actor is loaded, then invoke BuildMegaMesh on it
+		FWorldPartitionReference MegaMeshActorRef(WorldPartition.Get(), MegaMeshGUID);
+
+		// Note: Not every MegaMesh GUID references a valid mesh partition actor -- for example, there may be stale compiled sections that reference the GUID of old mesh partitions that have since been deleted
+		if (const AMeshPartition* MegaMesh = Cast<AMeshPartition>(MegaMeshActorRef.GetActor()))
+		{
+			bool bResult = BuildMegaMesh(MegaMesh, BuildKey, *DescriptorCache, InPackageHelper);
+		}
+	}
+
+	// Remove any out-of-date CompiledSections
+	UE_LOG(LogMegaMeshEditor, Display, TEXT("Removing out-of-date compiled sections from World %s"), *InWorld->GetName());
+
+	TArray<UPackage*> PackagesToCleanup;
+	int32 TotalRemovedCompiledSections = 0;
+	TArray<TWeakObjectPtr<AActor>> ActorsToUnload;
+	FWorldPartitionHelpers::ForEachActorDescInstance(WorldPartition, MeshPartition::ACompiledSection::StaticClass(),
+	[InWorld, WorldPartition = WorldPartition, BuildKey, CompiledSectionActorDescGuidsToKeep = CompiledSectionActorDescGuidsToKeep, &TotalRemovedCompiledSections, &PackagesToCleanup, &ActorsToUnload](const FWorldPartitionActorDescInstance* ActorDescInstance)
+	{
+		const FGuid& ActorDescGuid = ActorDescInstance->GetGuid();
+		if (CompiledSectionActorDescGuidsToKeep.Contains(ActorDescGuid))
+		{
+			// we're keeping this existing compiled section as part of an incremental build, don't delete it.
+			return true;
+		}
+
+		// if the actor has a build key, it is a compiled section we should check
+		if (const MeshPartition::FCompiledSectionActorDesc* CompiledSectionActorDesc = MeshPartition::FCompiledSectionActorDesc::GetFromActorDescInstance(*ActorDescInstance))
+		{
+			if (CompiledSectionActorDesc->GetBuildInfo().BuildKey == BuildKey)
+			{
+				// this actor was just built by this builder, don't delete it.
+				return true;
+			}
+		}
+
+		// the actor is not in our keep list, and was not built by this builder -- if it is a compiled section, then delete it
+		if (ActorDescInstance->GetActorNativeClass() == MeshPartition::ACompiledSection::StaticClass())
+		{
+			{
+				// try to load the compiled section into memory (must be loaded to delete it...)
+				FWorldPartitionReference CompiledSectionActorRef(WorldPartition.Get(), ActorDescGuid);
+				MeshPartition::ACompiledSection* CompiledSection = CastChecked<MeshPartition::ACompiledSection>(ActorDescInstance->GetActor());
+				if (CompiledSection != nullptr)					
+				{
+					const MeshPartition::FCompiledSectionBuildInfo& BuildInfo = CompiledSection->GetBuildInfo();
+					UE_LOG(LogMegaMeshEditor, Display, TEXT("  Deleting compiled section %.6s (BuildKey %.8s, %d Modifiers, Hash %.8s, BuildVariant %s %.8s"),
+						*ActorDescGuid.ToString(),
+						*BuildInfo.BuildKey.ToString(),
+						BuildInfo.BaseModifierPaths.Num(),
+						*BuildInfo.ModifiersHash.ToString(),
+						*BuildInfo.BuildVariantName.ToString(),
+						*BuildInfo.BuildVariantHash.ToString()
+					);
+
+					check(BuildInfo.BuildKey != BuildKey); // shouldn't be deleting any actors that were just built
+					PackagesToCleanup.Emplace(CompiledSection->GetPackage());
+					ActorsToUnload.Add(CompiledSection);
+					InWorld->DestroyActor(CompiledSection);
+					TotalRemovedCompiledSections++;
+				}
+				// closing the scope releases the reference
+			}
+
+			// run a garbage collection tick to evict the deleted actor and related resources
+			if (FWorldPartitionHelpers::ShouldCollectGarbage())
+			{
+				// cleanup any pending deleted packages
+				ObjectTools::CleanupAfterSuccessfulDelete(PackagesToCleanup, true);
+				PackagesToCleanup.Empty();
+
+				// Simulate an engine tick to make sure engine & render resources that are queued for deletion are processed.
+				FWorldPartitionHelpers::FakeEngineTick(InWorld);
+				FWorldPartitionHelpers::DoCollectGarbage();
+
+				// Double check that our deleted actors were actually garbage collected
+				UE::MeshPartition::EditorUtils::ValidateObjectsAreUnloaded(ActorsToUnload);
+				ActorsToUnload.Empty();
+			}
+		}
+
+		return true;
+	});
+	
+	// cleanup any remaining deleted packages
+	if (PackagesToCleanup.Num() > 0)
+	{
+		ObjectTools::CleanupAfterSuccessfulDelete(PackagesToCleanup, true);
+		PackagesToCleanup.Empty();
+	}
+	
+	// If the platform cell transformer is not yet installed to the world and we have compiled sections, install it now:
+	if ((TotalBuiltCompiledSections + TotalReusedCompiledSections) > 0)
+	{
+		if (!WorldPartition->HasRuntimeCellTransformerOfType(MeshPartition::UPlatformCellTransformer::StaticClass()))
+		{
+			WorldPartition->AddRuntimeCellTransformerOfType(TSubclassOf<UWorldPartitionRuntimeCellTransformer>(MeshPartition::UPlatformCellTransformer::StaticClass()));
+			SavePackages( { WorldPartition->GetPackage() }, InPackageHelper);
+		}
+	}
+
+	UE_LOG(LogMegaMeshEditor, Display, TEXT("Finished Build for World %s (Built %d, Reused %d, Failed %d, Removed %d Compiled Sections)"), *InWorld->GetName(),
+		TotalBuiltCompiledSections,
+		TotalReusedCompiledSections,
+		TotalFailedCompiledSections,
+		TotalRemovedCompiledSections);
+
+	return true;
+}
+
+bool UWorldPartitionMeshPartitionBuilder::BuildMegaMesh(const AMeshPartition* MegaMesh, const FGuid& BuildKey, MeshPartition::FModifierDescriptorCache& DescriptorCache, FPackageSourceControlHelper& InPackageHelper)
+{
+	TRACE_CPUPROFILER_EVENT_SCOPE(UWorldPartitionMeshPartitionBuilder::BuildMegaMesh);
+
+	const FGuid& MegaMeshGUID = MegaMesh->GetActorGuid();
+	if (!MegaMeshGUID.IsValid())
+	{
+		return false;
+	}
+
+	UWorld* World = MegaMesh->GetWorld();
+	check(World);
+
+	// Create a new build info -- this gets recorded on all compiled sections produced here.
+	MeshPartition::FCompiledSectionBuildInfo BuildInfo;
+	BuildInfo.BuildKey = BuildKey;
+	BuildInfo.MegaMeshGUID = MegaMeshGUID;
+	BuildInfo.MegaMeshPath = MegaMesh;
+
+	TSharedPtr<const MeshPartition::FCachedDescriptors> Descriptors = DescriptorCache.GetCachedDescriptors(MegaMesh);
+	if (!Descriptors.IsValid())
+	{
+		UE_LOG(LogMegaMeshEditor, Warning, TEXT("Could not gather descriptors for the Mesh Partition '%s'"), *MegaMesh->GetActorNameOrLabel());
+		return false;
+	}
+	TArrayView<const FCompiledSectionDescriptor> OldCompiledSections = Descriptors->GetCompiledSections();
+
+	// Use FWorldPartitionReference to load/unload everything in scope
+	FWorldPartitionReference MegaMeshActorRef(WorldPartition.Get(), MegaMeshGUID);
+	check(MegaMesh == CastChecked<AMeshPartition>(MegaMeshActorRef.GetActor()));
+
+	const UMeshPartitionDefinition* MeshPartitionDefinition = MegaMesh->GetMeshPartitionDefinition();
+	if (MeshPartitionDefinition == nullptr)
+	{
+		UE_LOG(LogMegaMeshEditor, Verbose, TEXT("No Mesh Partition Definition provided for %s (%.6s), using default settings."), *MegaMesh->GetActorNameOrLabel(), *MegaMeshGUID.ToString());
+		MeshPartitionDefinition = UMeshPartitionDefinition::StaticClass()->GetDefaultObject<UMeshPartitionDefinition>();
+	}
+	BuildInfo.SetMegaMeshDefinition(MeshPartitionDefinition);
+
+	TConstArrayView<MeshPartition::FCompiledSectionBuildVariant> BuildVariants = MeshPartitionDefinition->GetCompiledSectionBuildVariants();
+	check(BuildVariants.Num() > 0);
+
+	UMeshPartitionEditorComponent* MegaMeshEditorComponent = Cast<UMeshPartitionEditorComponent>(MegaMesh->GetMeshPartitionComponent());
+	// avoid level streaming triggering a preview section build while we compiled section resources are being built.
+	MegaMeshEditorComponent->SetPreviewSectionBuildEnabled(false);
+
+	const FTransform& Transform = MegaMesh->GetTransform();
+	check(MegaMeshEditorComponent != nullptr);
+
+	UE_LOG(LogMegaMeshEditor, Display, TEXT("-- %.6s Build Mesh Partition %s (%d Variants, %d Existing Compiled Sections"),
+		*MegaMeshGUID.ToString(), *MegaMesh->GetActorNameOrLabel(),
+		BuildVariants.Num(), OldCompiledSections.Num());
+
+	if (AMeshPartitionDataLayerContainer* Container = AMeshPartitionDataLayerContainer::GetOrCreate(World))
+	{
+		TArray<UPackage*> PackagesToSave;
+		if (Container->UpdateDataLayersFromDefinition(MeshPartitionDefinition))
+		{
+			PackagesToSave.Add(Container->GetPackage());
+		}
+		if (Container->InitializeDataLayerInstancesWithWorld(MeshPartitionDefinition))
+		{
+			PackagesToSave.Add(World->GetWorldDataLayers()->GetPackage());
+		}
+
+		SavePackages(PackagesToSave, InPackageHelper);
+	}
+
+	// Iterate over all the BuildVariants of the MegaMesh
+	for (const MeshPartition::FCompiledSectionBuildVariant& BuildVariant : BuildVariants)
+	{
+		TRACE_CPUPROFILER_EVENT_SCOPE_TEXT(
+			*FString::Format(TEXT("UWorldPartitionMeshPartitionBuilder::BuildMegaMeshVariant {0}"), {BuildVariant.Name.ToString()}));
+
+		// Record the build variant, so at cook time we can filter to the proper variant for each platform
+		BuildInfo.BuildVariantName = BuildVariant.Name;
+
+		FDependencyHash DependencyHash;
+		BuildInfo.GetMegaMeshDefinition()->GatherDependencies(DependencyHash, BuildVariant.Name);
+		BuildInfo.BuildVariantHash = DependencyHash.GetDependentDataHash();
+
+		UE_LOG(LogMegaMeshEditor, Display, TEXT("-- %.6s -- %s Build Variant"), *MegaMeshGUID.ToString(), *BuildVariant.Name.ToString());
+
+		TArrayView<const MeshPartition::FModifierGroup> ModifierGroups = Descriptors->GetModifierGroups(BuildVariant);
+
+		MegaMeshEditorComponent->ResetGroupRegistry(ModifierGroups);
+
+		MeshPartition::FBuilderSettings Settings;
+		Settings.BuildType = MeshPartition::EBuildType::CompiledSection;
+		Settings.Transform = Transform;
+		Settings.TypePriorities = MeshPartitionDefinition->GetModifierTypePriorities();
+		Settings.MaxSectionComplexity = BuildVariant.MaxSectionComplexity;
+		Settings.bRecomputeNormals = true;
+		Settings.bRecomputeTangents = true;
+		Settings.TexcoordGenerationOptions = FChannelCollectionUVLayoutOptions::GetFromDefinition(MeshPartitionDefinition);
+		// Choose the higher quality layout method for compiled sections
+		Settings.TexcoordGenerationOptions->UVLayoutMethod = EChannelCollectionUVLayoutMethod::ReferenceBoxProject;
+
+		// Using an indirect array so that the structures have a stable memory address
+		TIndirectArray<FModifierGroupTask> ModifierGroupTasks;
+
+		// Construct a ModifierGroupTask structure for each modifier group, to track the build state
+		for (const MeshPartition::FModifierGroup& ModifierGroup : ModifierGroups)
+		{
+			BuildInfo.BaseModifierPaths.Empty();
+			Algo::Transform(ModifierGroup.BaseDescs(), BuildInfo.BaseModifierPaths, [](const MeshPartition::FModifierDesc& BaseDescriptor) { return BaseDescriptor.ModifierPath; });
+
+			FModifierGroupTask* ModifierGroupTask = new FModifierGroupTask(ModifierGroup, Settings, BuildInfo, BuildVariant, MegaMeshEditorComponent, *this, OldCompiledSections);
+			ModifierGroupTasks.Add(ModifierGroupTask);
+			ModifierGroupTask->State = ETaskState::Ready;
+		}
+
+		ExecuteModifierGroupTasks(ModifierGroupTasks, World, MegaMesh, InPackageHelper);
+
+		int32 BuiltCompiledSections = 0;
+		int32 FailedCompiledSections = 0;
+		int32 ReusedCompiledSections = 0;
+		for (FModifierGroupTask& ModifierGroupTask : ModifierGroupTasks)
+		{
+			switch (ModifierGroupTask.Result)
+			{
+			case ETaskResult::BuildSuccess:
+				BuiltCompiledSections++;
+				break;
+			case ETaskResult::ReuseExisting:
+				ReusedCompiledSections++;
+				check(ModifierGroupTask.ReuseCompiledSectionActorDescGuid.IsValid());
+				CompiledSectionActorDescGuidsToKeep.Add(ModifierGroupTask.ReuseCompiledSectionActorDescGuid);
+				break;
+			default:
+				FailedCompiledSections++;
+				break;
+			}
+		}
+
+		// destruct the tasks to free up the memory and actor references they are holding
+		ModifierGroupTasks.Empty();
+		MegaMeshEditorComponent->ResetGroupRegistry();
+
+		UE_LOG(LogMegaMeshEditor, Display, TEXT("-- %.6s -- %s Finish Variant (Built %d, Reused %d, Failed %d Compiled Sections)"),
+			   *MegaMeshGUID.ToString(), *BuildVariant.Name.ToString(),
+			   BuiltCompiledSections, ReusedCompiledSections, FailedCompiledSections);
+
+		// Update stats
+		TotalBuiltCompiledSections += BuiltCompiledSections;
+		TotalReusedCompiledSections += ReusedCompiledSections;
+		TotalFailedCompiledSections += FailedCompiledSections;
+
+		// tick the engine and garbage collect to evict actors
+		FWorldPartitionHelpers::FakeEngineTick(World);
+		FWorldPartitionHelpers::DoCollectGarbage();
+	}
+
+	UE_LOG(LogMegaMeshEditor, Display, TEXT("-- %.6s -- Finished Mesh Partition Build (%s)"), *MegaMeshGUID.ToString(), *MegaMesh->GetActorNameOrLabel());
+	return true;
+}
+
+bool UWorldPartitionMeshPartitionBuilder::CanBuildMegaMeshes(const UWorld* InWorld, FName InBuildOption)
+{
+	return true;
+}
+
+EEditorBuildResult UWorldPartitionMeshPartitionBuilder::BuildMegaMeshes(UWorld* InWorld, EMeshPartitionReuseMode ReuseMode)
+{
+	check(InWorld);
+
+	// first let's run some quick sanity checks to warn upfront if things are incorrect
+	if (!InWorld->IsPartitionedWorld(InWorld))
+	{
+		const static FText MessageTitle = FText::FromString("Level is not using World Partition");
+		FString MessageString = FString::Printf(TEXT("The Mesh Partition Builder requires World Partition to be enabled. Conversion to World Partition is available via \"Tools > Convert Level...\"."));
+		FText MessageText = FText::FromString(MessageString);
+		FMessageDialog::Open(EAppMsgCategory::Warning, EAppMsgType::Ok, EAppReturnType::Ok, MessageText, MessageTitle);
+		return EEditorBuildResult::Skipped;
+	}
+
+	// check for definitions
+	for (TActorIterator<AMeshPartition> MegaMeshIterator(InWorld); MegaMeshIterator; ++MegaMeshIterator)
+	{
+		const AMeshPartition* MeshPartition = *MegaMeshIterator;
+		check(MeshPartition != nullptr);
+		
+		const UMeshPartitionDefinition* MeshPartitionDefinition = MeshPartition->GetMeshPartitionDefinition();
+		if (MeshPartitionDefinition == nullptr)
+		{
+			const static FText MessageTitle = FText::FromString("No Mesh Partition Definition provided");
+			FString MessageString = FString::Printf(TEXT("No Mesh Partition Definition provided for '%s'. Build using default settings?"), *MeshPartition->GetActorNameOrLabel());
+			FText MessageText = FText::FromString(MessageString);
+			EAppReturnType::Type Result = FMessageDialog::Open(EAppMsgCategory::Warning, EAppMsgType::OkCancel, EAppReturnType::Ok, MessageText, MessageTitle);
+			if (Result == EAppReturnType::Cancel)
+			{
+				return EEditorBuildResult::Skipped;
+			}
+		}
+		else
+		{
+			TConstArrayView<MeshPartition::FCompiledSectionBuildVariant> BuildVariants = MeshPartitionDefinition->GetCompiledSectionBuildVariants();
+			check(BuildVariants.Num() > 0);
+		}
+	}
+
+	const FString& PackageToReloadOnSuccess = GetNameSafe(InWorld->GetPackage());
+
+	// Try to provide the complete path. If we can't, try with project name.
+	const FString ProjectPath = FPaths::IsProjectFilePathSet() ? FPaths::GetProjectFilePath() : FApp::GetProjectName();
+
+	const ISourceControlProvider& SCCProvider = ISourceControlModule::Get().GetProvider();
+
+	FString ReuseModeString = UEnum::GetValueAsString(ReuseMode);
+	ReuseModeString = ReuseModeString.RightChop(ReuseModeString.Find("::") + 2);
+
+	const FString LogFileName = TEXT("WPMeshPartitionBuilderLog.txt");
+	const FString LogFileArg = TEXT("-log=") + LogFileName;
+
+	const FString Arguments = FString::Printf(TEXT("\"%s\" -run=WorldPartitionBuilderCommandlet %s %s -ReuseMode=%s %s -SCCProvider=%s%s"),
+		*ProjectPath,
+		*PackageToReloadOnSuccess,
+		TEXT("-AllowCommandletRendering -Builder=WorldPartitionMeshPartitionBuilder"),
+		*ReuseModeString,
+		*LogFileArg,
+		*SCCProvider.GetName().ToString(),
+		CVarMegaMeshBuildCommandWaitForDebugger->GetBool() ? TEXT(" -waitforattach") : TEXT(""));
+
+	FString FullLogPath = FPaths::ProjectLogDir() + LogFileName;
+	IPlatformFile& PlatformFile = FPlatformFileManager::Get().GetPlatformFile();
+	FDateTime LastCreationTime;
+	int64 LastReadPos = 0;
+
+	const FFileStatData FileStat = PlatformFile.GetStatData(*FullLogPath);
+	LastCreationTime = (FileStat.bIsValid) ? FileStat.CreationTime : FDateTime();
+	LastReadPos = (FileStat.bIsValid) ? FileStat.FileSize : 0;
+	
+	// Parse output, look for progress indicator in the log (in the form "Display: [i / N] Msg...\n")
+	const FRegexPattern LogProgressPattern(TEXT("Display:\\s\\[([0-9]+)\\s\\/\\s([0-9]+)\\]\\s(.+)?(?=\\.{3})"));
+
+	auto OnProcRunningTick = [&FullLogPath, &LastCreationTime, &LastReadPos, &LogProgressPattern, &PlatformFile](FScopedSlowTask& SlowTask)
+	{
+		const FFileStatData FileStat = PlatformFile.GetStatData(*FullLogPath);
+
+		if (!FileStat.bIsValid)
+		{
+			return;
+		}
+
+		const bool bLogRotated = (FileStat.CreationTime > LastCreationTime) || (FileStat.FileSize < LastReadPos);
+
+		if (bLogRotated)
+		{
+			LastCreationTime = FileStat.CreationTime;
+			LastReadPos = 0;
+		}
+
+		if (FileStat.FileSize > LastReadPos)
+		{
+			TUniquePtr<FArchive> LogReader(IFileManager::Get().CreateFileReader(*FullLogPath, FILEREAD_AllowWrite));
+
+			if (LogReader != nullptr)
+			{
+				TArray<uint8> Buffer;
+
+				LogReader->Seek(LastReadPos);
+				Buffer.SetNumUninitialized(FileStat.FileSize - LastReadPos);
+				LogReader->Serialize(Buffer.GetData(), Buffer.Num());
+
+				FString LogString = FString(StringCast<TCHAR>(reinterpret_cast<const char*>(Buffer.GetData())));
+				FRegexMatcher Regex(LogProgressPattern, LogString);
+
+				while (Regex.FindNext())
+				{
+					// Update slow task progress & message
+					SlowTask.CompletedWork = FCString::Atoi(*Regex.GetCaptureGroup(1));
+					SlowTask.TotalAmountOfWork = FCString::Atoi(*Regex.GetCaptureGroup(2));
+					SlowTask.DefaultMessage = FText::FromString(Regex.GetCaptureGroup(3));
+				}
+
+				LastReadPos = FileStat.FileSize;
+			}
+		}
+	};
+
+	const bool bSuccess = FEditorBuildUtils::RunWorldPartitionBuilder(PackageToReloadOnSuccess,
+		LOCTEXT("WorldPartitionBuildMeshPartitionProgress", "Building Mesh Partition..."),
+		LOCTEXT("WorldPartitionBuildMeshPartitionCancelled", "Building Mesh Partition cancelled!"),
+		LOCTEXT("WorldPartitionBuildMeshPartitionFailed", "Errors occured during the build process, please refer to the logs ('WPMeshPartitionBuilderLog.txt')."),
+		Arguments, OnProcRunningTick);
+
+	return bSuccess ? EEditorBuildResult::Success : EEditorBuildResult::Skipped;
+}
+
+void UWorldPartitionMeshPartitionBuilder::GatherDependenciesAndComputeHashes(FModifierGroupTask& InOutTask, IAssetRegistry& AssetRegistry)
+{
+	using namespace EditorUtils;
+	FDependencyContext DependencyContext;
+
+	// gather dependencies from the megamesh definition
+	if (const UMeshPartitionDefinition* MegaMeshDefinition = InOutTask.BuildInfo.GetMegaMeshDefinition())
+	{
+		DependencyContext.AddPackageDependency(MegaMeshDefinition);
+		DependencyContext.AddClassDependency(MegaMeshDefinition->GetClass());
+		MegaMeshDefinition->GatherDependencies(DependencyContext, InOutTask.BuildVariant.Name);
+	}
+
+	// gather dependencies from the megamesh modifiers
+	InOutTask.ModifierGroup.ForAllModifiers([&DependencyContext](MeshPartition::UModifierComponent* Modifier)
+	{
+		if (ensure(Modifier != nullptr))
+		{
+			DependencyContext.AddPackageDependency(Modifier);
+			DependencyContext.AddClassDependency(Modifier->GetClass());
+			Modifier->GatherDependencies(DependencyContext);
+		}
+
+		return true;
+	});
+
+	// populate the package tracking
+	InOutTask.BuildInfo.PackageDependencies = DependencyContext.GetPackageDependencies(AssetRegistry);
+	InOutTask.BuildInfo.PackageHash = ComputePackageHash(AssetRegistry, InOutTask.BuildInfo.PackageDependencies, &InOutTask.BuildInfo.PackageChecksums);
+
+	// populate the class tracking
+	TArray<const UClass*> Classes = DependencyContext.GetClassDependencies();
+	Algo::Transform(Classes, InOutTask.BuildInfo.ClassDependencies, [](const UClass* Class)
+		{
+			FSoftClassPath Path(Class);
+			return FName(Path.ToString());
+		});
+	InOutTask.BuildInfo.ClassHash = ComputeClassHash(Classes, &InOutTask.BuildInfo.ClassChecksums);
+}
+
+void UWorldPartitionMeshPartitionBuilder::ExecuteModifierGroupTasks(TIndirectArray<FModifierGroupTask>& ModifierGroupTasks, UWorld* InWorld, const AMeshPartition* InMegaMesh, FPackageSourceControlHelper& InPackageHelper)
+{
+	using namespace EditorUtils;
+
+	FAssetRegistryModule& AssetRegistryModule = FModuleManager::LoadModuleChecked<FAssetRegistryModule>("AssetRegistry");
+	IAssetRegistry& AssetRegistry = AssetRegistryModule.Get();
+
+	if ((ReuseMode == EMeshPartitionReuseMode::ByPackageHash) || (ReuseMode == EMeshPartitionReuseMode::ValidateHashes))
+	{
+		// we need to ensure the asset registry is up to date before we start querying it
+		// (otherwise we may fail to update things that were recently changed, but have not yet been updated in the asset registry...)
+		AssetRegistry.WaitForCompletion();
+	}
+
+	bool AllTaskCompleted = false;
+	int32 SectionsBuilt = 0;
+	int32 SectionsReused = 0;
+	bool bUpdateLog = false;
+
+	auto GetTotalSectionNumber = [&]()
+	{
+		return TotalBuiltCompiledSections + TotalFailedCompiledSections + TotalReusedCompiledSections + SectionsBuilt + SectionsReused;
+	};
+
+	int NumTasksInFlight = 0;
+
+	while (!AllTaskCompleted)
+	{
+		AllTaskCompleted = true;
+
+		for (FModifierGroupTask& Task : ModifierGroupTasks)
+		{
+			switch (Task.State)
+			{
+				case ETaskState::Ready:
+					if (NumTasksInFlight < MaxTasksInFlight && !FWorldPartitionHelpers::HasExceededMaxMemory())
+					{
+						Task.ReuseCompiledSectionActorDescGuid = FGuid();
+
+						// TODO [chris.tchou] : this DOES NOT include GUIDs from modifiers inside level instances...	so it doesn't detect adding a new instance inside a level instance.
+						Task.BuildInfo.ModifierSetHash = Task.ModifierGroup.ComputeModifierSetHash();
+
+						// try to reuse by package hashes (which doesn't require loading any modifiers, so we can early out if this works)
+						if ((ReuseMode == EMeshPartitionReuseMode::ByPackageHash) || (ReuseMode == EMeshPartitionReuseMode::ValidateHashes))
+						{
+							if (const FCompiledSectionDescriptor* MatchingCompiledSection = FindMatchingCompiledSectionFromPackageHashes(Task.OldCompiledSectionDescriptors, Task.BuildInfo, AssetRegistry))
+							{
+								// early out, reuse existing
+								Task.ReuseCompiledSectionActorDescGuid = MatchingCompiledSection->ActorDescGuid;
+							}
+						}
+
+						// if we're not early out (or if we're in validate mode)
+						if (!Task.ReuseCompiledSectionActorDescGuid.IsValid() || (ReuseMode == EMeshPartitionReuseMode::ValidateHashes))
+						{
+							// load the modifiers, including all modifiers within level instances
+							LoadActors(InMegaMesh, Task);
+							Task.State = ETaskState::ActorLoaded;
+
+							// and compute an up-to-date modifier hash
+							Task.BuildInfo.ModifiersHash = Task.ModifierGroup.UpdateAndComputeModifierGroupHash();
+
+							if (ReuseMode == EMeshPartitionReuseMode::ByModifierHash)
+							{
+								Task.ReuseCompiledSectionActorDescGuid = FindMatchingCompiledSectionFromModifierHash(Task.OldCompiledSectionDescriptors, Task.BuildInfo);
+							}
+							else if (ReuseMode == EMeshPartitionReuseMode::ValidateHashes)
+							{
+								FGuid ModifierHashReuseGuid = FindMatchingCompiledSectionFromModifierHash(Task.OldCompiledSectionDescriptors, Task.BuildInfo);
+
+								// ModifierHash may find a match when PackageHash does not, that is expected.
+								// But we can check:
+
+								// if PackageHash finds a match, then modifier hash should agree
+								if (Task.ReuseCompiledSectionActorDescGuid.IsValid())
+								{
+									// (if not we have a bug in our hash/dependency calculations)
+									check(Task.ReuseCompiledSectionActorDescGuid == ModifierHashReuseGuid);
+								}
+
+								// if ModifierHash finds NO match, then PackageHash should agree
+								if (!ModifierHashReuseGuid.IsValid())
+								{
+									// (if not we have a bug in our hash/dependency calculations)
+									check(!Task.ReuseCompiledSectionActorDescGuid.IsValid());
+								}
+
+								// TODO: if we record PackageHash and/or ModifierHash per Package/Modifier, then we can provide better info on which is causing the issue here
+							}
+						}
+
+						if (Task.ReuseCompiledSectionActorDescGuid.IsValid())
+						{
+							Task.Result = ETaskResult::ReuseExisting;
+							Task.State = ETaskState::Completed;
+							SectionsReused++;
+							bUpdateLog = true;
+
+							WorldPartitionMegaMeshBuilderLocals::LogBuildProgress(GetTotalSectionNumber(), TotalNumberOfSectionsToBuild);
+						}
+						else
+						{
+							// record package and class dependencies
+							GatherDependenciesAndComputeHashes(Task, AssetRegistry);
+
+							// Build a new compiled section.  Start by launching the async task to execute the modifiers
+							Task.ModifierGroup.SetBuildType(MeshPartition::EBuildType::CompiledSection);
+							Task.ModifierGroup.ProgressToState(MeshPartition::FModifierGroup::EState::BackgroundOpsCreated);
+
+							Task.ModifierTaskGraph = MakeShared<MeshPartition::FModifierTaskGraph>();
+
+							Task.State = ETaskState::LaunchedProcessModifierTask;
+							// #todo: this whole processing task should be using the MeshPartition builder system instead of using its own bespoke path.
+							Tasks::FTask ProcessModifiersTask = Tasks::Launch(TEXT("UWorldPartitionMeshPartitionBuilder_ProcessModifiersTask"),
+								[&Task, Group = Task.ModifierGroup.CreateAsyncBuildGroup()]() mutable
+								{
+									TRACE_CPUPROFILER_EVENT_SCOPE(UWorldPartitionMeshPartitionBuilder::ProcessModifiersTask);
+
+									Tasks::FTask ProcessComplete = BuildHelpers::ProcessModifierGroup(Task.BuilderSettings, Group, Task.ModifierTaskGraph);
+
+									Tasks::AddNested(ProcessComplete);
+								}
+							);
+						
+							Tasks::TTask<MeshPartition::FMeshData> GenerateUVLayoutTask = Tasks::Launch(TEXT("UWorldPartitionMeshPartitionMeshBuilder_GenerateUVLayoutTask"),
+								[&Task]
+								{
+									MeshPartition::FMeshData MeshData = MoveTemp(Task.ModifierTaskGraph->GetResultMesh());
+
+									if (Task.BuilderSettings.TexcoordGenerationOptions.IsSet())
+									{
+										MeshPartition::FChannelTextureRenderer::CreateSectionUVLayout(MeshData, Task.BuilderSettings.TexcoordGenerationOptions.GetValue());
+									}
+									return MoveTemp(MeshData);
+								},
+								Tasks::Prerequisites(ProcessModifiersTask)
+							);
+
+							Task.ProcessModifierTask = Tasks::Launch(TEXT("UWorldPartitionMeshPartitionMeshBuilder_FinishProcessModifiersTask"),
+								[&Task, GenerateUVLayoutTask]() mutable
+								{
+									MeshPartition::FMeshData MeshData = MoveTemp(GenerateUVLayoutTask.GetResult());
+
+									Task.State = ETaskState::ModifiersProcessed;
+									return MoveTemp(MeshData);
+								},
+								Tasks::Prerequisites(GenerateUVLayoutTask)
+							);
+
+							++NumTasksInFlight;
+						}
+					}
+					break;
+								
+				case ETaskState::ModifiersProcessed:
+					Task.CompiledSection = Task.MegaMeshEditorComponent->PrepareMegaMeshCompiledSection(Task.BuildInfo, Task.BuildVariant);
+					Task.MegaMeshEditorComponent->BuildMegaMeshCompiledSectionTextures(Task.CompiledSection, Task.ModifierGroup, Task.ProcessModifierTask.GetResult());
+
+					{
+						TArray<TWeakObjectPtr<MeshPartition::UModifierComponent>> ModifierPtrs = Task.ModifierGroup.AllResolvedModifierPtrs();
+						Task.MegaMeshEditorComponent->PostBuildSectionMesh(Task.CompiledSection, Task.ProcessModifierTask.GetResult(), ModifierPtrs);
+
+						MeshPartition::FTransformerUnit InitialUnit;
+
+						InitialUnit.Section = Task.CompiledSection;
+						InitialUnit.MeshData = MakeShared<const MeshPartition::FMeshData>(MoveTemp(Task.ProcessModifierTask.GetResult()));
+						InitialUnit.bShouldRecomputeNormals = true;
+						InitialUnit.bShouldRecomputeTangents = true;
+
+						MeshPartition::UMeshPartitionDefinition* Definition = Task.MegaMeshEditorComponent->GetMegaMeshDefinition();
+
+						Task.TransformerContext = Task.MegaMeshEditorComponent->LaunchTransformers(MoveTemp(InitialUnit), Definition, Task.BuildVariant);
+					}
+
+					if (Task.TransformerContext.IsValid())
+					{
+						Task.State = ETaskState::TransformersLaunched;
+					}
+					else
+					{
+						Task.State = ETaskState::TransformersExecuted;
+					}
+
+					break;
+					
+				case ETaskState::TransformersLaunched:
+					if (Task.TransformerContext->JoinTask.IsValid() && Task.TransformerContext->JoinTask.IsCompleted())
+					{
+						Task.State = ETaskState::TransformersExecuted;
+					}
+
+					break;
+
+				case ETaskState::TransformersExecuted:
+					Task.ModifiedPackages.Emplace(Task.CompiledSection->GetPackage());
+
+					{
+						TArray<TWeakObjectPtr<MeshPartition::UModifierComponent>> ModifierPtrs = Task.ModifierGroup.AllResolvedModifierPtrs();
+						Task.MegaMeshEditorComponent->PostProcessSection(Task.CompiledSection, ModifierPtrs);
+					}
+
+					AssignSectionToPlatformDataLayers(Task.CompiledSection);
+
+					// save and reclaim memory
+					SaveAndCollectGarbage(InWorld, InPackageHelper, Task);
+					Task.Result = ETaskResult::BuildSuccess;
+					Task.State = ETaskState::Completed;
+					SectionsBuilt++;
+					bUpdateLog = true;
+					NumTasksInFlight--;
+
+					int32 TotalSectionNumber = TotalBuiltCompiledSections + TotalFailedCompiledSections + TotalReusedCompiledSections + SectionsBuilt + SectionsReused;
+					WorldPartitionMegaMeshBuilderLocals::LogBuildProgress(GetTotalSectionNumber(), TotalNumberOfSectionsToBuild);
+
+					break;
+			}
+			
+			if (Task.State != ETaskState::Completed)
+			{
+				AllTaskCompleted = false;
+			}
+		}
+		
+		if (bUpdateLog)
+		{
+			bUpdateLog = false;
+			UE_LOG(LogMegaMeshEditor, Log, TEXT("Sections: %d built, %d reused out of %d total"), SectionsBuilt, SectionsReused, ModifierGroupTasks.Num());
+		}
+		
+		// Tick the engine while we still have remaining tasks to progress engine systems such as PCG and Rendering, allowing them to do work on the GameThread if needed.
+		FWorldPartitionHelpers::FakeEngineTick(InWorld);
+	}
+}
+
+void UWorldPartitionMeshPartitionBuilder::LoadActors(const AMeshPartition* InMegaMesh, FModifierGroupTask& InModifierGroupTask)
+{
+	TRACE_CPUPROFILER_EVENT_SCOPE(UWorldPartitionMeshPartitionBuilder::LoadActors);
+
+	if (InMegaMesh == nullptr)
+	{
+		return;
+	}
+
+	UMeshPartitionEditorComponent* MegaMeshEditorComponent = Cast<UMeshPartitionEditorComponent>(InMegaMesh->GetMeshPartitionComponent());
+	if (MegaMeshEditorComponent == nullptr)
+	{
+		return;
+	}
+
+	auto OnActorLoaded = [this, &InModifierGroupTask](AActor* Actor, bool bIsInLevelInstance)
+		{
+			InModifierGroupTask.ModifierOwners.Emplace(Actor);
+		};
+
+
+	auto OnModifierLoaded = [this, MegaMeshEditorComponent, &InModifierGroupTask](MeshPartition::UModifierComponent* Modifier, bool bIsInLevelInstance)
+		{
+			MegaMeshEditorComponent->OnModifierLoaded(Modifier);
+		};
+
+	WorldPartitionHelpers::LoadAllRelevantMegaMeshActors(InModifierGroupTask.ModifierGroup, InMegaMesh, InModifierGroupTask.ActorRefs, OnActorLoaded, OnModifierLoaded);
+}
+
+void UWorldPartitionMeshPartitionBuilder::AssignSectionToPlatformDataLayers(MeshPartition::ACompiledSection* InCompiledSection) const
+{
+	const UWorld* World = InCompiledSection->GetWorld();
+	if (AMeshPartitionDataLayerContainer* Container =  AMeshPartitionDataLayerContainer::Get(World))
+	{
+		if (TObjectPtr<UDataLayerAsset> DataLayer = Container->FindVariantDataLayer(InCompiledSection->GetParentMegaMesh()->GetMeshPartitionDefinition(), InCompiledSection->GetBuildInfo().BuildVariantName); ensure(DataLayer))
+		{
+			if (UDataLayerManager* DataLayerManager = UDataLayerManager::GetDataLayerManager(World); ensure(DataLayerManager))
+			{
+				if (const UDataLayerInstance* Instance = DataLayerManager->GetDataLayerInstance(DataLayer); ensure(Instance))
+				{
+					Instance->AddActor(InCompiledSection);
+				}
+			}
+		}
+	}
+}
+
+bool UWorldPartitionMeshPartitionBuilder::SavePackages(const TArray<UPackage*>& PackagesToSave, FPackageSourceControlHelper& InPackageHelper)
+{
+	TArray<FString> PackageNames;
+	Algo::Transform(PackagesToSave, PackageNames, [](const UPackage* Package) { return Package->GetName(); });
+
+	// first figure out whether the files need to be checked out or added to source control
+	TArray<FString> PackageNamesToCheckout;
+	TArray<FString> PackageNamesToAdd;
+	if (!InPackageHelper.GetDesiredStatesForModification(PackageNames, PackageNamesToCheckout, PackageNamesToAdd))
+	{
+		const FString FailedPackages = FString::Join(PackageNames, TEXT(", "));
+		UE_LOG(LogMegaMeshEditor, Error, TEXT("Cannot save packages because files are exclusively checked out or not at head revision: %s."), *FailedPackages);
+		return false;
+	}
+
+	// second, checkout all of the files that need to be checked out (as a single operation)
+	if (!InPackageHelper.Checkout(PackageNamesToCheckout))
+	{
+		const FString FailedPackages = FString::Join(PackageNamesToCheckout, TEXT(", "));
+		UE_LOG(LogMegaMeshEditor, Error, TEXT("Error checking out packages: %s."), *FailedPackages);
+		return false;
+	}
+
+	// third, save all packages to disk
+	bool bResult = true;
+	TSet<FString> FailedToSaveNames;
+	for (UPackage* Package : PackagesToSave)
+	{
+		FString PackageFileName = SourceControlHelpers::PackageFilename(Package);
+		FSavePackageArgs SaveArgs;
+		SaveArgs.TopLevelFlags = RF_Standalone;
+
+		if (!UPackage::SavePackage(Package, nullptr, *PackageFileName, SaveArgs))
+		{
+			UE_LOG(LogMegaMeshEditor, Error, TEXT("Error saving package %s."), *Package->GetName());
+			FailedToSaveNames.Add(Package->GetName());
+			bResult = false;
+		}
+	}
+
+	// fourth, add packages to source control that need to be added (as a single operation)
+	// we do this even if some packages failed to save, so that we don't leave files on disk that were not added to source control
+	{
+		// produce a list of packages to add, filtering out those that failed to save
+		TArray<FString> FilteredPackageNamesToAdd;
+		Algo::TransformIf(PackageNamesToAdd, FilteredPackageNamesToAdd,
+			[&FailedToSaveNames](const FString& PackageName) { return !FailedToSaveNames.Contains(PackageName); },
+			[](const FString& PackageName) { return PackageName; });
+
+		if (!InPackageHelper.AddToSourceControl(FilteredPackageNamesToAdd))
+		{
+			const FString FailedPackages = FString::Join(FilteredPackageNamesToAdd, TEXT(", "));
+			UE_LOG(LogMegaMeshEditor, Error, TEXT("Error adding packages to source control: %s."), *FailedPackages);
+			bResult = false;
+		}
+	}
+
+	return bResult;
+}
+
+void UWorldPartitionMeshPartitionBuilder::SaveAndCollectGarbage(UWorld* InWorld, FPackageSourceControlHelper& InPackageHelper, FModifierGroupTask& InModifierGroupTask)
+{
+	TRACE_CPUPROFILER_EVENT_SCOPE(UWorldPartitionMeshPartitionBuilder::SaveAndCollectGarbage);
+
+	TArray<UObject*> ObjectsToFinish;
+	ForEachObjectWithOuter(InModifierGroupTask.CompiledSection, [&ObjectsToFinish](UObject* Object)
+	{
+		if (const IInterface_AsyncCompilation* AsyncCompilationIF = Cast<IInterface_AsyncCompilation>(Object))
+		{
+			ObjectsToFinish.Emplace(Object);
+		}
+	
+		return true;
+	});
+
+	if (!ObjectsToFinish.IsEmpty())
+	{
+		FAssetCompilingManager::Get().FinishCompilationForObjects(ObjectsToFinish);
+	}
+
+	if (!InModifierGroupTask.ModifiedPackages.IsEmpty())
+	{
+		SavePackages(InModifierGroupTask.ModifiedPackages, InPackageHelper);
+	}
+
+	// Destroy the newly created compiled section _after_ the save so that it can be fully unloaded by the next GC.
+	InWorld->DestroyActor(InModifierGroupTask.CompiledSection);
+
+	// clear out pointers, since the objects may be going away in the garbage collection
+	InModifierGroupTask.CompiledSection = nullptr;
+	InModifierGroupTask.ModifiedPackages.Empty();
+	InModifierGroupTask.ModifierOwners.Empty();
+	InModifierGroupTask.ActorRefs.Empty();
+
+	// Simulate an engine tick to make sure engine & render resources that are queued for deletion are processed and streaming levels are removed.
+	FWorldPartitionHelpers::FakeEngineTick(InWorld);
+	FWorldPartitionHelpers::DoCollectGarbage();
+}
+} // namespace UE::MeshPartition
+
+#undef LOCTEXT_NAMESPACE
diff --git a/Engine/Plugins/Experimental/MeshPartition/Source/MeshPartitionEditor/Public/MeshPartitionPlatformCellTransformer.h b/Engine/Plugins/Experimental/MeshPartition/Source/MeshPartitionEditor/Public/MeshPartitionPlatformCellTransformer.h
new file mode 100644
--- /dev/null
+++ b/Engine/Plugins/Experimental/MeshPartition/Source/MeshPartitionEditor/Public/MeshPartitionPlatformCellTransformer.h
@@ -0,0 +1,40 @@
+// Copyright Epic Games, Inc. All Rights Reserved.
+#pragma once
+
+#include "WorldPartition/WorldPartitionRuntimeCellTransformer.h"
+#include "MeshPartitionPlatformCellTransformer.generated.h"
+
+class ULevel;
+class UDataLayerAsset;
+class IWorldPartitionCell;
+
+namespace UE::MeshPartition
+{
+
+class AMeshPartitionDataLayerContainer;
+
+USTRUCT()
+struct FStrippedDataLayerList
+{
+	GENERATED_BODY()
+
+	UPROPERTY()
+	TArray<TObjectPtr<const UDataLayerAsset>> List;
+};
+
+UCLASS(meta = (DisplayName = "Mesh Partition Platform Cell Transformer"))
+class UPlatformCellTransformer : public UWorldPartitionRuntimeCellTransformer
+{
+	GENERATED_BODY()
+public:
+	virtual void TransformRuntimeCell(UWorldPartitionRuntimeCell* InCell) override;
+	virtual void Transform(ULevel* InLevel) override;
+
+protected:
+	static const MeshPartition::AMeshPartitionDataLayerContainer* GetDataLayerContainerForCell(const UWorldPartitionRuntimeCell* InCell);
+
+private:
+	UPROPERTY(NonPIEDuplicateTransient)
+	TMap<FGuid, FStrippedDataLayerList> CellToStrippedDataLayers;
+};
+} // namespace UE::MeshPartition
\ No newline at end of file
diff --git a/Engine/Plugins/Experimental/MeshPartition/Source/MeshPartitionEditor/Public/WorldPartitionMeshPartitionBuilder.h b/Engine/Plugins/Experimental/MeshPartition/Source/MeshPartitionEditor/Public/WorldPartitionMeshPartitionBuilder.h
new file mode 100644
--- /dev/null
+++ b/Engine/Plugins/Experimental/MeshPartition/Source/MeshPartitionEditor/Public/WorldPartitionMeshPartitionBuilder.h
@@ -0,0 +1,210 @@
+// Copyright Epic Games, Inc. All Rights Reserved.
+
+#pragma once
+
+#include "MeshPartitionModifierComponent.h"
+#include "MeshPartitionCompiledSection.h"
+#include "MeshPartitionEditorComponent.h" // FModifierDescriptorPair
+#include "MeshPartitionModifierTaskGraph.h"
+#include "Tasks/Task.h"
+#include "WorldPartition/WorldPartition.h"
+#include "WorldPartition/WorldPartitionBuilder.h"
+
+#include "WorldPartitionMeshPartitionBuilder.generated.h"
+
+class AActor;
+class UPackage;
+class UWorld;
+class UWorldPartition;
+enum class EEditorBuildResult : uint8;
+
+namespace UE::MeshPartition
+{
+	struct FModifierDescriptorCache;
+	class ACompiledSection;
+	struct FCompiledSectionBuildInfo;
+
+	UENUM()
+	enum class EMeshPartitionReuseMode
+	{
+		ForceRebuild,			// slowest: rebuild all compiled sections, even if we don't think it is necessary
+		ByModifierHash,			// fast:	skip building compiled sections where the modifier cachekey hash has not changed (requires loading the modifiers to calculate their cache keys)
+		ByPackageHash,			// fastest: does not require loading modifiers. skip building compiled sections where the dependent asset packages (as declared by GatherDependencies) have not changed.
+		ValidateHashes,			// fast:	runs both package hash and modifier hash paths, and validates their behavior.  Used mainly for debugging.
+		Count UMETA(Hidden)
+	};
+}
+
+// ENUM_RANGE_BY_COUNT defines a symbol that needs to be declared in the global namespace.
+ENUM_RANGE_BY_COUNT(UE::MeshPartition::EMeshPartitionReuseMode, UE::MeshPartition::EMeshPartitionReuseMode::Count);
+
+namespace UE::MeshPartition
+{
+// Example Command Line: ProjectName MapName -run=WorldPartitionBuilderCommandlet -SCCProvider=Perforce -Builder=WorldPartitionMegaMeshBuilder
+UCLASS()
+class UWorldPartitionMeshPartitionBuilder : public UWorldPartitionBuilder
+{
+	GENERATED_UCLASS_BODY()
+
+	enum class ETaskState
+	{
+		Invalid,
+		/** Task is ready to execute */
+		Ready,
+		/**	The modifier and base actor are loaded */
+		ActorLoaded,
+		/**	Launched the async process modifier task, but the task has not yet started */
+		LaunchedProcessModifierTask,
+		/**	The task is currently processing the modifiers, building the FDynamicMesh3 */
+		ProcessingModifiers,
+		/**	Modifiers are processed. Built mesh data is ready */
+		ModifiersProcessed,
+		/**	The Transformers pipeline is launched and processing. */
+		TransformersLaunched,
+		/**	The Transformers pipeline execution is done. */
+		TransformersExecuted,
+		/**	Task's work is done */
+		Completed
+	};
+
+	enum class ETaskResult
+	{
+		/* Initial state, not yet complete */
+		None,
+		/* Failed to build a new compiled section */
+		BuildFail,
+		/* Successfully built a new compiled section */
+		BuildSuccess,
+		/* Reused an existing compiled section */
+		ReuseExisting
+	};
+	
+	struct FModifierGroupTask
+	{
+		// Inputs
+		UE::MeshPartition::FModifierGroup		ModifierGroup;
+		MeshPartition::FBuilderSettings 		BuilderSettings;
+		MeshPartition::FCompiledSectionBuildInfo	BuildInfo;
+		const MeshPartition::FCompiledSectionBuildVariant&		BuildVariant;
+		UMeshPartitionEditorComponent*			MegaMeshEditorComponent = nullptr;
+		UWorldPartitionMeshPartitionBuilder&		WorldPartitionMegaMeshBuilder;		
+
+		// array of pre-existing compiled sections (potentially re-usable for incremental rebuild)
+		// NOTE: this is NOT thread safe, and only works because we always synchronously complete this task
+		// and deallocate it  before destroying the array this view references.
+		TArrayView<const FCompiledSectionDescriptor> OldCompiledSectionDescriptors;
+
+		// Results
+		ETaskResult							Result = ETaskResult::None;
+		MeshPartition::ACompiledSection*			CompiledSection = nullptr;
+		FGuid								ReuseCompiledSectionActorDescGuid;
+
+		// internal intermediate state
+		ETaskState							State = ETaskState::Invalid;
+
+		TUniquePtr<MeshPartition::FTransformerContext> TransformerContext;
+		UE::Tasks::TTask<FMeshData>				ProcessModifierTask;
+		TSharedPtr<MeshPartition::FModifierTaskGraph> 			ModifierTaskGraph;
+
+		TArray<UPackage*>					ModifiedPackages;
+		TSet<AActor*>						ModifierOwners;			// set of actors containing one of our modifier components
+		TSet<FWorldPartitionReference>		ActorRefs;				// references to world partition actors that must remain loaded while this task is running
+		
+		FModifierGroupTask(UE::MeshPartition::FModifierGroup&& InModifierGroup, const MeshPartition::FBuilderSettings& InBuilderSettings, const MeshPartition::FCompiledSectionBuildInfo& InBuildInfo, const MeshPartition::FCompiledSectionBuildVariant& InBuildVariant, UMeshPartitionEditorComponent* InMegaMeshEditorComponent, UWorldPartitionMeshPartitionBuilder& InWorldPartitionMegaMeshBuilder, TArrayView<const FCompiledSectionDescriptor> InOldCompiledSectionDescriptors)
+			: ModifierGroup(MoveTemp(InModifierGroup))
+			, BuilderSettings(InBuilderSettings)
+			, BuildInfo(InBuildInfo)
+			, BuildVariant(InBuildVariant)
+			, MegaMeshEditorComponent(InMegaMeshEditorComponent)
+			, WorldPartitionMegaMeshBuilder(InWorldPartitionMegaMeshBuilder)
+			, OldCompiledSectionDescriptors(InOldCompiledSectionDescriptors)
+		{}
+
+		FModifierGroupTask(const UE::MeshPartition::FModifierGroup& InModifierGroup, const MeshPartition::FBuilderSettings& InBuilderSettings, const MeshPartition::FCompiledSectionBuildInfo& InBuildInfo, const MeshPartition::FCompiledSectionBuildVariant& InBuildVariant, UMeshPartitionEditorComponent* InMegaMeshEditorComponent, UWorldPartitionMeshPartitionBuilder& InWorldPartitionMegaMeshBuilder, TArrayView<const FCompiledSectionDescriptor> InOldCompiledSectionDescriptors)
+			: ModifierGroup(InModifierGroup)
+			, BuilderSettings(InBuilderSettings)
+			, BuildInfo(InBuildInfo)
+			, BuildVariant(InBuildVariant)
+			, MegaMeshEditorComponent(InMegaMeshEditorComponent)
+			, WorldPartitionMegaMeshBuilder(InWorldPartitionMegaMeshBuilder)
+			, OldCompiledSectionDescriptors(InOldCompiledSectionDescriptors)
+		{}
+	};
+
+public:
+	// UWorldPartitionBuilder interface begin
+	virtual bool RequiresCommandletRendering() const override { return false; }
+	virtual ELoadingMode GetLoadingMode() const override { return ELoadingMode::Custom; }
+	virtual bool PreRun(UWorld* World, FPackageSourceControlHelper& PackageHelper) override;
+	virtual bool RunInternal(UWorld* World, const FCellInfo& InCellInfo, FPackageSourceControlHelper& InPackageHelper) override;
+	virtual bool CanProcessNonPartitionedWorlds() const override { return false; }
+	// UWorldPartitionBuilder interface end
+
+	static bool CanBuildMegaMeshes(const UWorld* InWorld, FName InBuildOption);
+	static EEditorBuildResult BuildMegaMeshes(UWorld* InWorld, EMeshPartitionReuseMode ReuseMode);
+
+private:
+	/**
+	* Build a MegaMesh
+	*/
+	bool BuildMegaMesh(const AMeshPartition* MegaMesh, const FGuid& BuildKey, UE::MeshPartition::FModifierDescriptorCache& DescriptorCache, FPackageSourceControlHelper& InPackageHelper);
+
+	/**
+	* Executes and manages all modifier group tasks upon completion.
+	* @param InWorld The current world.
+	* @param InMegaMesh The AMegaMesh the modifiers belong to.
+	* @param InMegaMeshEditorComponent The UMeshPartitionEditorComponent performing the processing.
+	* @param InPackageHelper A helper for saving and source control operations.
+	*/
+	void ExecuteModifierGroupTasks(TIndirectArray<FModifierGroupTask>& ModifierGroupTasks, UWorld* InWorld, const AMeshPartition* InMegaMesh, FPackageSourceControlHelper& InPackageHelper);
+
+	/**
+	* Loads the modifier owners. Populating the lookup table and refcounting the loaded modifier owners.
+	* @param InMegaMesh The AMegaMesh the modifiers belong to.
+	* @param InModifierGroupTask The ModifierGroupTask that needs processing.
+	*/
+	void LoadActors(const AMeshPartition* InMegaMesh, FModifierGroupTask& InModifierGroupTask);
+
+	void AssignSectionToPlatformDataLayers(MeshPartition::ACompiledSection* InCompiledSection) const;
+
+	/**
+	* Gather the set of all package dependencies from the modifiers in the ModifierGroup.
+	* @param AssetRegistry the asset registry we use to track recursive dependencies.  Must be up to date before calling this function for correct results.
+	* @param ModifierGroup the modifier group to gather dependencies from
+	* @return An array of package names for packages containing dependencies of the modifier group, sorted in a deterministic order.
+	*/
+	void GatherDependenciesAndComputeHashes(FModifierGroupTask& InOutTask, IAssetRegistry& AssetRegistry);
+
+	/**
+	* Helper to save a package.
+	* @param PackagesToSave The package to save.
+	* @param InPackageHelper A helper for saving and source control operations.
+	* @return 
+	*/
+	static bool SavePackages(const TArray<UPackage*>& PackagesToSave, FPackageSourceControlHelper& InPackageHelper);
+	
+	/**
+	* Waits for all compilable before saving and cleaning resources if memory is needed and the resources are not needed anymore.
+	* @param InWorld The current world.
+	* @param InPackageHelper A helper for saving and source control operations.
+	* @param InModifierGroupTask The ModifierGroupTask that needs saving.
+	*/
+	void SaveAndCollectGarbage(UWorld* InWorld, FPackageSourceControlHelper& InPackageHelper, FModifierGroupTask& InModifierGroupTask);
+
+private:
+	UPROPERTY(meta=(CommandLineArgument))
+	EMeshPartitionReuseMode ReuseMode = EMeshPartitionReuseMode::ByPackageHash;
+
+	UPROPERTY(meta=(CommandLineArgument))
+	int32 MaxTasksInFlight = 16;
+
+	TObjectPtr<UWorldPartition> WorldPartition = nullptr;
+
+	// Build state
+	TSet<FGuid> CompiledSectionActorDescGuidsToKeep;
+	int32 TotalBuiltCompiledSections = 0;
+	int32 TotalFailedCompiledSections = 0;
+	int32 TotalReusedCompiledSections = 0;
+	int32 TotalNumberOfSectionsToBuild = 0;
+};
+} // namespace UE::MeshPartition
\ No newline at end of file
diff --git a/Engine/Plugins/Experimental/MeshPartition/Source/MeshPartitionModelingToolset/Private/MeshPartitionModelingToolCommands.cpp b/Engine/Plugins/Experimental/MeshPartition/Source/MeshPartitionModelingToolset/Private/MeshPartitionModelingToolCommands.cpp
new file mode 100644
--- /dev/null
+++ b/Engine/Plugins/Experimental/MeshPartition/Source/MeshPartitionModelingToolset/Private/MeshPartitionModelingToolCommands.cpp
@@ -0,0 +1,65 @@
+// Copyright Epic Games, Inc. All Rights Reserved.
+
+#include "MeshPartitionModelingToolCommands.h"
+
+#include "MeshPartitionModelingToolsStyle.h"
+#include "MeshPartitionHeightSculptTool.h"
+
+#define LOCTEXT_NAMESPACE "MeshModelingToolsCommands"
+
+FMegaMeshModelingToolCommands::FMegaMeshModelingToolCommands()
+	: TCommands<FMegaMeshModelingToolCommands>(
+		"MegaMeshModelingToolCommands", // Context name for fast lookup
+		NSLOCTEXT("Contexts", "MegaMeshModelingToolCommands", "Mega Mesh Modeling Tools"), // Localized context name for displaying
+		NAME_None, // Parent
+		FMegaMeshModelingToolsStyle::Get()->GetStyleSetName() // Icon Style Set
+		)
+{
+}
+
+void FMegaMeshModelingToolCommands::RegisterCommands()
+{
+	UI_COMMAND(MegaMeshToolsTabButton, "MegaMesh", "Mega Mesh Modeling Toolset", EUserInterfaceActionType::RadioButton, FInputChord());
+	UI_COMMAND(BeginConvertMeshTool, "Convert", "Convert the target mesh into a new Mega Mesh actor", EUserInterfaceActionType::ToggleButton, FInputChord());
+	UI_COMMAND(BeginSplitMeshTool, "Split", "Split a mesh into multiple Mega Mesh section actors", EUserInterfaceActionType::ToggleButton, FInputChord());
+	UI_COMMAND(BeginMergeMeshTool, "Merge", "Merges multiple Mega Mesh sections into a single section actor", EUserInterfaceActionType::ToggleButton, FInputChord());
+	UI_COMMAND(BeginResectionMeshTool, "Resection", "Merges and then divides selected mesh sections", EUserInterfaceActionType::ToggleButton, FInputChord());
+	UI_COMMAND(BeginStitchMeshTool, "Stitch", "Stitch a regular mesh into a Mega Mesh section", EUserInterfaceActionType::ToggleButton, FInputChord());
+	UI_COMMAND(BeginHeightmapImport, "Import Heightmap", "Create a new Mega Mesh from a heightmap", EUserInterfaceActionType::ToggleButton, FInputChord());
+	UI_COMMAND(BeginHeightSculptTool, "Height Sculpt", "Sculpt height based on a reference surface", EUserInterfaceActionType::ToggleButton, FInputChord());
+	UI_COMMAND(BeginExpandMeshTool, "Expand", "Insert new Mega Mesh sections at section boundaries", EUserInterfaceActionType::ToggleButton, FInputChord());
+	UI_COMMAND(BeginCreateMegaMeshRectangleTool, "Create Rectangle", "Create a new rectangular Mega Mesh", EUserInterfaceActionType::ToggleButton, FInputChord());
+	UI_COMMAND(BeginAddModifierTool, "Add Modifier", "Create a new Mega Mesh modifier", EUserInterfaceActionType::ToggleButton, FInputChord());
+}
+
+//~ Modeled on ModelingToolsActions.cpp
+UE::Geometry::FMegaMeshToolActionCommands::FMegaMeshToolActionCommands() :
+	TInteractiveToolCommands<FMegaMeshToolActionCommands>(
+		"MegaMeshHotkeys", // Context name for fast lookup
+		LOCTEXT("HotkeysCategory", "MegaMesh Hotkeys"), // Localized context name for displaying
+		NAME_None, // Parent
+		FMegaMeshModelingToolsStyle::Get()->GetStyleSetName() // Icon Style Set
+	)
+{
+}
+void UE::Geometry::FMegaMeshToolActionCommands::RegisterAllToolActions()
+{
+	UE::Geometry::FMegaMeshHeightSculptToolCommands::Register();
+}
+void UE::Geometry::FMegaMeshToolActionCommands::UnregisterAllToolActions()
+{
+	UE::Geometry::FMegaMeshHeightSculptToolCommands::Unregister();
+}
+
+//~ Modeled on ModelingToolsActions.cpp
+#define DEFINE_TOOL_ACTION_COMMANDS(CommandsClassName, ContextNameString, SettingsDialogString, ToolClassName ) \
+UE::Geometry::CommandsClassName::CommandsClassName() : TInteractiveToolCommands<CommandsClassName>( \
+ContextNameString, NSLOCTEXT("Contexts", ContextNameString, SettingsDialogString), NAME_None, FMegaMeshModelingToolsStyle::Get()->GetStyleSetName()) {} \
+void UE::Geometry::CommandsClassName::GetToolDefaultObjectList(TArray<UInteractiveTool*>& ToolCDOs) \
+{\
+	ToolCDOs.Add(GetMutableDefault<ToolClassName>()); \
+}
+
+DEFINE_TOOL_ACTION_COMMANDS(FMegaMeshHeightSculptToolCommands, "HeightSculpt", "MegaMesh - Height Sculpt", UE::MeshPartition::UHeightSculptTool);
+
+#undef LOCTEXT_NAMESPACE
diff --git a/Engine/Plugins/Experimental/MeshTerrainMode/MeshTerrainMode.uplugin b/Engine/Plugins/Experimental/MeshTerrainMode/MeshTerrainMode.uplugin
new file mode 100644
--- /dev/null
+++ b/Engine/Plugins/Experimental/MeshTerrainMode/MeshTerrainMode.uplugin
@@ -0,0 +1,62 @@
+{
+	"FileVersion": 3,
+	"Version": 1,
+	"VersionName": "0.1",
+	"FriendlyName": "Mesh Terrain Mode",
+	"Description": "Mesh Terrain Mode includes a suite of interactive tools for creating and editing Mesh Terrain in the Editor",
+	"Category": "Editor",
+	"CreatedBy": "Epic Games, Inc.",
+	"CreatedByURL": "https://epicgames.com",
+	"DocsURL": "",
+	"MarketplaceURL": "",
+	"SupportURL": "",
+	"CanContainContent": true,
+	"EnabledByDefault": false,
+	"IsBetaVersion": false,
+	"IsExperimentalVersion": true,
+	"Installed": false,
+	"Modules": [
+		{
+			"Name": "MeshTerrainMode",
+			"Type": "Editor",
+			"LoadingPhase": "Default"
+		}
+	],
+	"Plugins": [
+		{
+			"Name": "MeshModelingToolset",
+			"Enabled": true,
+			"TargetAllowList": [ "Editor" ]
+		},
+		{
+			"Name": "MeshModelingToolsetExp",
+			"Enabled": true,
+			"TargetAllowList": [ "Editor" ]
+		},
+		{
+			"Name": "MeshLODToolset",
+			"Enabled": true,
+			"TargetAllowList": [ "Editor" ]
+		},
+		{
+			"Name": "ToolPresets",
+			"Enabled": true,
+			"TargetAllowList": [ "Editor" ]
+		},
+		{
+			"Name": "StylusInput",
+			"Enabled": true,
+			"TargetAllowList": [ "Editor" ]
+		},
+		{
+			"Name": "ModelingToolsEditorMode",
+			"Enabled": true,
+			"TargetAllowList": [ "Editor" ]
+		},
+		{
+			"Name": "MeshPartition",
+			"Enabled": true,
+			"TargetAllowList": [ "Editor" ]
+		}
+	]
+}
\ No newline at end of file
