[Mesh Partition]

Updated the build hash to fix some bad ddc.

#Synced-CL 51480380

diff --git a/Engine/Plugins/Experimental/MeshPartition/Source/MeshPartitionEditor/Private/MeshPartitionMeshBuilder.cpp b/Engine/Plugins/Experimental/MeshPartition/Source/MeshPartitionEditor/Private/MeshPartitionMeshBuilder.cpp
--- a/Engine/Plugins/Experimental/MeshPartition/Source/MeshPartitionEditor/Private/MeshPartitionMeshBuilder.cpp
+++ b/Engine/Plugins/Experimental/MeshPartition/Source/MeshPartitionEditor/Private/MeshPartitionMeshBuilder.cpp
@@ -383,7 +383,7 @@ namespace UE::MeshPartition::BuildHelpers
 			Hasher.Update(Name.GetData(), Name.Len() * sizeof(TCHAR));
 		};
 
-		static FGuid VersionKey(TEXT("fc71e455-e6cf-46d7-aa1f-7f4c31fc4573"));
+		static FGuid VersionKey(TEXT("51e0a251-3d12-4852-a6c8-c1b1092faacd"));
 		UpdateHashPODType(Hasher, VersionKey);
 
 		// Hash the current mesh data type version into the key to ensure swapping mesh types hits different DDC entries.
