{MeshPartition] UV unwrap and uv seams improvments
- Improve the uv unwrap  by doing a first evaluation of the total 3d area of the mesh and from this target an ideal texture size.
from this target size derive proper gutter size required in the uv unwrap process
- FIxed this for both  UV unwrap methods

- Fixed the uv seam artefact on nanite mesh. this is due to a degraded "No derivative" mode of rendering the mesh because of the  the material texture sample of the channel texture array.
- We are working aournd the problem by uising an explicit TextureSample providing the derivative explicitely (with no variation on the slice component) which fix the issue in the rendering
- The mesh is apprently still using the"NoDerivative" path for nanite and i will look into this next.

- Improved the border fill pass and provided an alternative technique "Push pull".
- Modified the existing BorderFIll env to indicate a technique (like for uv unwrap) 0 means no border fill, 1 means the current technique and 2 uses the Push pull technique
- Improved the existing shader for better quality.

- Exported useful information in a dedicated struct FMeshPartitionChannelGenerationMesh which is assigned to the PreviewSection. This will come handy to create an inspector to look at the channels of the section and understand artefacts.

- Cleaned up the draw channel shader to make it dead simple

- Removed some more MegaMesh from the code

#Synced-CL 51490229

diff --git a/Engine/Plugins/Experimental/MeshPartition/Shaders/MeshPartitionBorderFill.usf b/Engine/Plugins/Experimental/MeshPartition/Shaders/MeshPartitionBorderFill.usf
--- a/Engine/Plugins/Experimental/MeshPartition/Shaders/MeshPartitionBorderFill.usf
+++ b/Engine/Plugins/Experimental/MeshPartition/Shaders/MeshPartitionBorderFill.usf
@@ -4,7 +4,7 @@
 
 uint2 Resolution;
 
-Texture2D<uint> Mask;
+Texture2D<float> Mask;
 
 RWTexture2D<float4> RWSectionTexture;
 
@@ -16,23 +16,24 @@ void MainCS(uint2 DTid : SV_DispatchThreadID)
 		return;
 	}
 
-	// Naive filling, assume all mask candidates are valid
+	const int2 TexPos = int2(DTid);
 	
-	const int BoxSearchExtent         = 4;
-	const int MaxCandidateDistanceSqr = Square(BoxSearchExtent);
+	// From the kernel size derive the search box for candidate texels
+	const int KernelSize = 3;
+	const int2 SearchMinOffset = max(TexPos - KernelSize, 0);
+	const int2 SearchMaxOffset = min(TexPos + KernelSize, int2(Resolution - 1));
 
-	const uint2 SearchMinOffset = max(int2(DTid) - BoxSearchExtent, 0);
-	const uint2 SearchMaxOffset = min(int2(DTid) + BoxSearchExtent, int2(Resolution - 1));
-
-	int    CandidateDistanceSqr = MaxCandidateDistanceSqr;
+	// candidate samples within the kernel radius can contribute
+	int BestCandidateDistanceSqr = Square(KernelSize + 1);
 	float4 CandidateTexel       = 0.0f;
+	int    NumSamples			= 0;
 
-	for (uint CoordY = SearchMinOffset.y; CoordY <= SearchMaxOffset.y; CoordY++)
+	for (int CoordY = SearchMinOffset.y; CoordY <= SearchMaxOffset.y; CoordY++)
 	{
-		for (uint CoordX = SearchMinOffset.x; CoordX <= SearchMaxOffset.x; CoordX++)
+		for (int CoordX = SearchMinOffset.x; CoordX <= SearchMaxOffset.x; CoordX++)
 		{
-			int DistanceSqr = Square(int(CoordX) - int(DTid.x)) + Square(int(CoordY) - int(DTid.y));
-			if (DistanceSqr >= CandidateDistanceSqr)
+			int DistanceSqr = Square(CoordX - TexPos.x) + Square(CoordY - TexPos.y);
+			if (DistanceSqr > BestCandidateDistanceSqr)
 			{
 				continue;
 			}
@@ -43,13 +44,86 @@ void MainCS(uint2 DTid : SV_DispatchThreadID)
 				continue;
 			}
 
-			CandidateDistanceSqr = DistanceSqr;
-			CandidateTexel       = RWSectionTexture[Offset];
+			// reset the candidate sample accumulation if closer candidate found 
+			if (DistanceSqr < BestCandidateDistanceSqr)
+			{
+				BestCandidateDistanceSqr = DistanceSqr;
+				NumSamples = 0;
+				CandidateTexel = 0;
+			}
+
+			// accumulate sample to candidate texel and num samples
+			CandidateTexel += RWSectionTexture[Offset];
+			NumSamples++;
 		}
 	}
 
-	if (CandidateDistanceSqr != MaxCandidateDistanceSqr)
+	if (NumSamples > 0)
 	{
-		RWSectionTexture[DTid] = CandidateTexel;
+		RWSectionTexture[DTid] = CandidateTexel / float(NumSamples);
 	}
 }
+
+float2 TexelSize;
+int4 ResolutionPass;
+Texture2D<float4> SectionMipIn;
+RWTexture2D<float4> SectionMipOut;
+Texture2D<float> MaskMipIn;
+RWTexture2D<float> MaskMipOut;
+SamplerState MipSampler;
+
+[numthreads(8, 8, 1)]
+void PullCS(uint2 DTid : SV_DispatchThreadID)
+{
+	if (any(DTid >= ResolutionPass.xy))
+	{
+		return;
+	}
+	
+	int2 DestTexel = DTid.xy;
+	int2 SourceTexel = DestTexel * 2;
+	float4 Signal = 0;
+	float Mask = 0;
+	for (int i = 0; i < 2; ++i)
+	{
+		for (int j = 0; j < 2; ++j)
+		{
+			float SourceMask = MaskMipIn.Load(int3(SourceTexel + int2(i, j), 0));
+			float4 SourceSignal = SectionMipIn.Load(int3(SourceTexel + int2(i, j), 0));
+			
+			Signal += SourceMask * SourceSignal;
+			Mask += SourceMask;
+		}
+	}
+	
+	if (Mask > 0)
+	{
+		Signal = Signal / Mask;
+		Mask *= 0.25;
+	}
+	else
+	{
+		Signal = 0;
+		Mask = 0;
+	}
+	
+	SectionMipOut[DestTexel] = Signal;
+	MaskMipOut[DestTexel] = Mask;
+}
+
+[numthreads(8, 8, 1)]
+void PushCS(uint2 DTid : SV_DispatchThreadID)
+{
+	if (any(DTid >= ResolutionPass.xy) || MaskMipIn[DTid])
+	{
+		return;
+	}
+
+	float2 coord = (DTid.xy + (0.5)) * TexelSize;
+	
+	float4 Source = SectionMipIn.Sample(MipSampler, coord);
+	
+	//SectionMipOut[DTid] = float4(coord, 0, 1);
+	SectionMipOut[DTid] = Source;
+}
+
diff --git a/Engine/Plugins/Experimental/MeshPartition/Shaders/MeshPartitionMakeSectionChannels.usf b/Engine/Plugins/Experimental/MeshPartition/Shaders/MeshPartitionMakeSectionChannels.usf
--- a/Engine/Plugins/Experimental/MeshPartition/Shaders/MeshPartitionMakeSectionChannels.usf
+++ b/Engine/Plugins/Experimental/MeshPartition/Shaders/MeshPartitionMakeSectionChannels.usf
@@ -6,39 +6,32 @@
 
 Buffer<int>		InMeshIndices;	
 Buffer<float2>	InMeshUVs;
-Buffer<float>	InMeshChannelMasks;
+Buffer<float>	InMeshChannels;
 
-uint4	InDrawCall; // PageNum, NumChannelAttribs, FirstVertIndex
-uint4   InChannelsOffset; // Channel offsets from the very first channel attribute for this vertex
-float4  InChannelsMask;
+uint4  InDrawCall; // NumVertices, NumChannelAttribs, FirstVertIndex
+uint   InChannelOffset; // Offset to the first channel element in the Channel buffer
 
 // ----------------------------------------------------------------------------------
 // Vertex shaders : 
 
-void MegaMeshDrawUVDomainVS(
+void DrawUVDomainVS(
 	uint VertexIndex : SV_VertexID,
 	out float2 SourceUV : TEXCOORD0,
-	out float4  OutMasks : TEXCOORD1,
+	out float4  OutWeight : TEXCOORD1,
 	out float4 OutPosition : SV_POSITION)
 {
 	uint VID = InMeshIndices[VertexIndex + InDrawCall.z];
 	
 	SourceUV = InMeshUVs[VID];
 	
-	// Channel Attrib Address is at VID x NumChannelAttribs
-	uint ChannelAttribOffset = VID * InDrawCall.y;
+	// Channel Attrib Address is at ChannelID * NumVertices + VID
+	uint ChannelAttribOffset = InChannelOffset + VID;
 	
-	// Fetch the 4 channels from the channel buffer (and maybe overflow fetching the neighbour vert channel values)
-	OutMasks = float4(	InMeshChannelMasks[ChannelAttribOffset + InChannelsOffset[0]],
-						InMeshChannelMasks[ChannelAttribOffset + InChannelsOffset[1]],
-						InMeshChannelMasks[ChannelAttribOffset + InChannelsOffset[2]],
-						InMeshChannelMasks[ChannelAttribOffset + InChannelsOffset[3]]);
-	
-	// Channels not written by this pass need to be reset to 0
-	OutMasks *= InChannelsMask;
+	// Fetch the channel weight from the channel buffer
+	OutWeight = InMeshChannels[ChannelAttribOffset];
 
+	// Transform from UV space to ndc
 	float2 BaseUV = SourceUV;	
-	
 	BaseUV *= 2.0;
 	OutPosition.x = BaseUV.x - 1.0;
 	OutPosition.y = 1.0 - BaseUV.y;
@@ -48,13 +41,13 @@ void MegaMeshDrawUVDomainVS(
 // ----------------------------------------------------------------------------------
 // Pixel shaders : 
 
-void MegaMeshDrawUVDomainPS(
+void DrawUVDomainPS(
 	float2 SourceUV : TEXCOORD0,
-	float4 Masks : TEXCOORD1,
+	float4 Weights : TEXCOORD1,
 	float4 SvPosition : SV_POSITION,
 	out float4 OutColor : SV_Target0,
-	out uint OutMask : SV_Target1)
+	out float OutMask : SV_Target1)
 {
-	OutColor = float4(Masks);
-	OutMask  = 1;
+	OutColor = Weights;
+	OutMask  = 1.0;
 }
diff --git a/Engine/Plugins/Experimental/MeshPartition/Source/MeshPartition/Public/MeshPartitionDefinition.h b/Engine/Plugins/Experimental/MeshPartition/Source/MeshPartition/Public/MeshPartitionDefinition.h
--- a/Engine/Plugins/Experimental/MeshPartition/Source/MeshPartition/Public/MeshPartitionDefinition.h
+++ b/Engine/Plugins/Experimental/MeshPartition/Source/MeshPartition/Public/MeshPartitionDefinition.h
@@ -179,6 +179,9 @@ public:
 	MeshPartition::FChannelMap& GetChannelMap() { return ChannelMap; }
 	const MeshPartition::FChannelMap& GetChannelMap() const { return ChannelMap; }
 	
+	UFUNCTION(CallInEditor)
+	TArray<FName> GetChannelNames() const { return ChannelMap.GetChannels(); }
+
 	float GetChannelTexelSize() const { return ChannelTexelSize; }
 
 	const TArray<FPhysicalMaterialChannel>& GetPhysicalMaterialChannels() const { return PhysicalMaterialChannels; }
@@ -268,12 +271,6 @@ private:
 	TArray<MeshPartition::FCompiledSectionBuildVariant> BuildVariants_DEPRECATED;
 #endif // WITH_EDITORONLY_DATA
 
-	private:
-		UFUNCTION(CallInEditor)
-		TArray<FName> GetChannelNames() const
-		{
-			return ChannelMap.GetChannels();
-		}
 };
 } // namespace UE::MeshPartition
 
diff --git a/Engine/Plugins/Experimental/MeshPartition/Source/MeshPartition/Public/MeshPartitionMeshData.h b/Engine/Plugins/Experimental/MeshPartition/Source/MeshPartition/Public/MeshPartitionMeshData.h
--- a/Engine/Plugins/Experimental/MeshPartition/Source/MeshPartition/Public/MeshPartitionMeshData.h
+++ b/Engine/Plugins/Experimental/MeshPartition/Source/MeshPartition/Public/MeshPartitionMeshData.h
@@ -520,4 +520,17 @@ private:
 };
 
 using FMeshABBTree3 = Geometry::TMeshAABBTree3<FMeshData>;
+
+// Simple mesh data representation used during the generation of the channel textures
+// Used by the Preview Section in order to inspect the channel texture.
+struct FSectionChannelGenerationMeshData
+{
+	TArray<int32>			Indices;
+	TArray<FVector2f>		Texcoords;
+	TArray<float>			Channels;
+	double					DomainArea3D;
+	double					DomainAreaUV;
+	double					DomainGutterUV;
+};
+
 } // namespace UE::MeshPartition
diff --git a/Engine/Plugins/Experimental/MeshPartition/Source/MeshPartitionCompute/Private/MeshPartitionChannelRasterizationShaders.cpp b/Engine/Plugins/Experimental/MeshPartition/Source/MeshPartitionCompute/Private/MeshPartitionChannelRasterizationShaders.cpp
--- a/Engine/Plugins/Experimental/MeshPartition/Source/MeshPartitionCompute/Private/MeshPartitionChannelRasterizationShaders.cpp
+++ b/Engine/Plugins/Experimental/MeshPartition/Source/MeshPartitionCompute/Private/MeshPartitionChannelRasterizationShaders.cpp
@@ -2,7 +2,11 @@
 
 #include "MeshPartitionChannelRasterizationShaders.h"
 
-IMPLEMENT_GLOBAL_SHADER(FMegaMesh_DrawUVDomainVS, "/Plugin/MeshPartition/MeshPartitionMakeSectionChannels.usf", "MegaMeshDrawUVDomainVS", SF_Vertex);
-IMPLEMENT_GLOBAL_SHADER(FMegaMesh_DrawUVDomainPS, "/Plugin/MeshPartition/MeshPartitionMakeSectionChannels.usf", "MegaMeshDrawUVDomainPS", SF_Pixel);
+IMPLEMENT_GLOBAL_SHADER(FMeshPartition_DrawUVDomainVS, "/Plugin/MeshPartition/MeshPartitionMakeSectionChannels.usf", "DrawUVDomainVS", SF_Vertex);
+IMPLEMENT_GLOBAL_SHADER(FMeshPartition_DrawUVDomainPS, "/Plugin/MeshPartition/MeshPartitionMakeSectionChannels.usf", "DrawUVDomainPS", SF_Pixel);
+
+IMPLEMENT_GLOBAL_SHADER(FMeshPartition_BorderFillCS, "/Plugin/MeshPartition/MeshPartitionBorderFill.usf", "MainCS", SF_Compute);
+IMPLEMENT_GLOBAL_SHADER(FMeshPartition_FillPullCS, "/Plugin/MeshPartition/MeshPartitionBorderFill.usf", "PullCS", SF_Compute);
+IMPLEMENT_GLOBAL_SHADER(FMeshPartition_FillPushCS, "/Plugin/MeshPartition/MeshPartitionBorderFill.usf", "PushCS", SF_Compute);
+
 
-IMPLEMENT_GLOBAL_SHADER(FMegaMesh_BorderFillCS, "/Plugin/MeshPartition/MeshPartitionBorderFill.usf", "MainCS", SF_Compute);
diff --git a/Engine/Plugins/Experimental/MeshPartition/Source/MeshPartitionCompute/Public/MeshPartitionChannelRasterizationShaders.h b/Engine/Plugins/Experimental/MeshPartition/Source/MeshPartitionCompute/Public/MeshPartitionChannelRasterizationShaders.h
--- a/Engine/Plugins/Experimental/MeshPartition/Source/MeshPartitionCompute/Public/MeshPartitionChannelRasterizationShaders.h
+++ b/Engine/Plugins/Experimental/MeshPartition/Source/MeshPartitionCompute/Public/MeshPartitionChannelRasterizationShaders.h
@@ -10,38 +10,37 @@
 #define UE_API MESHPARTITIONCOMPUTE_API
 
 //~ Plugins/MeshPartition/Shaders/MegaMeshMakeSectionChannels.usf shaders :
-class FMegaMesh_DrawUVDomainVS : public FGlobalShader
+class FMeshPartition_DrawUVDomainVS : public FGlobalShader
 {
 public:
-	DECLARE_EXPORTED_GLOBAL_SHADER(FMegaMesh_DrawUVDomainVS, UE_API);
-	SHADER_USE_PARAMETER_STRUCT(FMegaMesh_DrawUVDomainVS, FGlobalShader);
+	DECLARE_EXPORTED_GLOBAL_SHADER(FMeshPartition_DrawUVDomainVS, UE_API);
+	SHADER_USE_PARAMETER_STRUCT(FMeshPartition_DrawUVDomainVS, FGlobalShader);
 
 	BEGIN_SHADER_PARAMETER_STRUCT(FParameters, )
 		SHADER_PARAMETER_RDG_BUFFER_SRV(Buffer<uint32>, InMeshIndices)
 		SHADER_PARAMETER_RDG_BUFFER_SRV(Buffer<float2>, InMeshUVs)
-		SHADER_PARAMETER_RDG_BUFFER_SRV(Buffer<float>, InMeshChannelMasks)
+		SHADER_PARAMETER_RDG_BUFFER_SRV(Buffer<float>, InMeshChannels)
 		SHADER_PARAMETER(FUintVector4, InDrawCall)
-		SHADER_PARAMETER(FUintVector4, InChannelsOffset)
-		SHADER_PARAMETER(FVector4f, InChannelsMask)
+		SHADER_PARAMETER(uint32, InChannelOffset)
 	END_SHADER_PARAMETER_STRUCT()
 };
 
-class FMegaMesh_DrawUVDomainPS : public FGlobalShader
+class FMeshPartition_DrawUVDomainPS : public FGlobalShader
 {
 public:
-	DECLARE_EXPORTED_GLOBAL_SHADER(FMegaMesh_DrawUVDomainPS, UE_API);
-	SHADER_USE_PARAMETER_STRUCT(FMegaMesh_DrawUVDomainPS, FGlobalShader);
+	DECLARE_EXPORTED_GLOBAL_SHADER(FMeshPartition_DrawUVDomainPS, UE_API);
+	SHADER_USE_PARAMETER_STRUCT(FMeshPartition_DrawUVDomainPS, FGlobalShader);
 
 	BEGIN_SHADER_PARAMETER_STRUCT(FParameters, )
 		RENDER_TARGET_BINDING_SLOTS()
 	END_SHADER_PARAMETER_STRUCT()
 };
 
-class FMegaMesh_BorderFillCS : public FGlobalShader
+class FMeshPartition_BorderFillCS : public FGlobalShader
 {
 public:
-	DECLARE_EXPORTED_GLOBAL_SHADER(FMegaMesh_BorderFillCS, UE_API);
-	SHADER_USE_PARAMETER_STRUCT(FMegaMesh_BorderFillCS, FGlobalShader);
+	DECLARE_EXPORTED_GLOBAL_SHADER(FMeshPartition_BorderFillCS, UE_API);
+	SHADER_USE_PARAMETER_STRUCT(FMeshPartition_BorderFillCS, FGlobalShader);
 
 	BEGIN_SHADER_PARAMETER_STRUCT(FParameters, )
 		SHADER_PARAMETER(FUintVector2, Resolution)
@@ -55,9 +54,53 @@ public:
 	}
 };
 
-BEGIN_SHADER_PARAMETER_STRUCT(FMegaMesh_DrawUVDomain_Parameters, )
-	SHADER_PARAMETER_STRUCT_INCLUDE(FMegaMesh_DrawUVDomainVS::FParameters, VS)
-	SHADER_PARAMETER_STRUCT_INCLUDE(FMegaMesh_DrawUVDomainPS::FParameters, PS)
+BEGIN_SHADER_PARAMETER_STRUCT(FMeshPartition_DrawUVDomain_Parameters, )
+	SHADER_PARAMETER_STRUCT_INCLUDE(FMeshPartition_DrawUVDomainVS::FParameters, VS)
+	SHADER_PARAMETER_STRUCT_INCLUDE(FMeshPartition_DrawUVDomainPS::FParameters, PS)
 END_SHADER_PARAMETER_STRUCT()
 
+
+class FMeshPartition_FillPullCS : public FGlobalShader
+{
+public:
+	DECLARE_EXPORTED_GLOBAL_SHADER(FMeshPartition_FillPullCS, UE_API);
+	SHADER_USE_PARAMETER_STRUCT(FMeshPartition_FillPullCS, FGlobalShader);
+
+	BEGIN_SHADER_PARAMETER_STRUCT(FParameters, )
+		SHADER_PARAMETER(FVector2f, TexelSize)
+		SHADER_PARAMETER(FIntVector4, ResolutionPass)
+		SHADER_PARAMETER_RDG_TEXTURE_SRV(Texture2D<float4>, SectionMipIn)
+		SHADER_PARAMETER_RDG_TEXTURE_UAV(RWTexture2D<float4>, SectionMipOut)
+		SHADER_PARAMETER_RDG_TEXTURE_SRV(Texture2D<float>, MaskMipIn)
+		SHADER_PARAMETER_RDG_TEXTURE_UAV(RWTexture2D<float>, MaskMipOut)
+		SHADER_PARAMETER_SAMPLER(SamplerState, MipSampler)
+		END_SHADER_PARAMETER_STRUCT()
+
+		static void ModifyCompilationEnvironment(const FGlobalShaderPermutationParameters& InParameters, FShaderCompilerEnvironment& OutEnvironment)
+	{
+		OutEnvironment.CompilerFlags.Add(CFLAG_AllowTypedUAVLoads);
+	}
+};
+
+class FMeshPartition_FillPushCS : public FGlobalShader
+{
+public:
+	DECLARE_EXPORTED_GLOBAL_SHADER(FMeshPartition_FillPushCS, UE_API);
+	SHADER_USE_PARAMETER_STRUCT(FMeshPartition_FillPushCS, FGlobalShader);
+
+	BEGIN_SHADER_PARAMETER_STRUCT(FParameters, )
+		SHADER_PARAMETER(FVector2f, TexelSize)
+		SHADER_PARAMETER(FIntVector4, ResolutionPass)
+		SHADER_PARAMETER_RDG_TEXTURE_SRV(Texture2D<float4>, SectionMipIn)
+		SHADER_PARAMETER_RDG_TEXTURE_UAV(RWTexture2D<float4>, SectionMipOut)
+		SHADER_PARAMETER_RDG_TEXTURE_SRV(Texture2D<float>, MaskMipIn)
+		SHADER_PARAMETER_SAMPLER(SamplerState, MipSampler )
+	END_SHADER_PARAMETER_STRUCT()
+
+		static void ModifyCompilationEnvironment(const FGlobalShaderPermutationParameters& InParameters, FShaderCompilerEnvironment& OutEnvironment)
+	{
+		OutEnvironment.CompilerFlags.Add(CFLAG_AllowTypedUAVLoads);
+	}
+};
+
 #undef UE_API
diff --git a/Engine/Plugins/Experimental/MeshPartition/Source/MeshPartitionEditor/Private/MeshPartitionChannelCollection.cpp b/Engine/Plugins/Experimental/MeshPartition/Source/MeshPartitionEditor/Private/MeshPartitionChannelCollection.cpp
--- a/Engine/Plugins/Experimental/MeshPartition/Source/MeshPartitionEditor/Private/MeshPartitionChannelCollection.cpp
+++ b/Engine/Plugins/Experimental/MeshPartition/Source/MeshPartitionEditor/Private/MeshPartitionChannelCollection.cpp
@@ -55,10 +55,21 @@ static TAutoConsoleVariable<int32> CVarMeshPartition_UnwrapUVsMethod(
 	})
 );
 
-static TAutoConsoleVariable<bool> CVarMeshPartition_BorderFill(
-	TEXT("MeshPartition.Channels.BorderFill"),
-	true,
-	TEXT("Enables border filling on section page rendering")
+static TAutoConsoleVariable<int32> CVarMeshPartition_BorderFillMethod(
+	TEXT("MeshPartition.Channels.BorderFillMethod"),
+	1,
+	TEXT("Enables border filling on section page rendering. 0: No border; 1: Border fill; 2: Push pull"),
+	FConsoleVariableDelegate::CreateLambda([](IConsoleVariable*)
+	{
+		UWorld* World = GEditor->GetEditorWorldContext().World();
+		for (UE::MeshPartition::AMeshPartition* MeshPartition : TActorRange<UE::MeshPartition::AMeshPartition>(World))
+		{
+			if (UE::MeshPartition::UMeshPartitionEditorComponent* EditorComponent = Cast<UE::MeshPartition::UMeshPartitionEditorComponent>(MeshPartition->GetMeshPartitionComponent()))
+			{
+				EditorComponent->ForceRebuildAllSections(UE::MeshPartition::EChangeType::TransientStateChange);
+			}
+		}
+	})
 );
 
 namespace UE::MeshPartition
@@ -185,6 +196,8 @@ Tasks::TTask<MeshPartition::FSectionChannels> FChannelTextureRenderer::BuildText
 				.Table = MoveTemp(Section.ChannelOffsetTable),
 				.TexcoordMetrics = Section.TexcoordMetrics,
 				.Texture = ChannelTexture,
+				.Mesh = { MoveTemp(Section.UVmesh.Indices), MoveTemp(Section.UVmesh.UVs), MoveTemp(Section.UVmesh.Channels),
+							Section.DomainMapping.Area3D, Section.DomainMapping.AreaUV, Section.DomainMapping.GetGutterUV() }
 			};
 
 			return Result;
@@ -202,73 +215,65 @@ void FChannelTextureRenderer::GenerateSectionMeshChannels(const FMeshData& InSec
 	TRACE_CPUPROFILER_EVENT_SCOPE(MeshPartition::FChannelTextureRenderer::GenerateSectionMeshChannels);
 
 	int32 NumVertices = InSection.UVmesh.NumVertices();
-	int32 NumChannelsInAttributes = InSectionMesh.GetWeightLayerNames().Num();
 	int32 MaxNumChannels = InChannels.GetNumChannels();
 
 	TArray<int32> ChannelActiveCounts;
 	ChannelActiveCounts.Init(0, MaxNumChannels);
 
-	// gather the section channels values
-	// and allocate new channel entry if needed
-	TArray<float> SectionChannels;
-	SectionChannels.Init(0, NumVertices * MaxNumChannels);
-
-	if (NumChannelsInAttributes > 0)
+	if (InSectionMesh.GetWeightLayerNames().Num() > 0)
 	{
+		// Grab the local packed array of channel attributes from the weight layer values
+		const TArray<FName> WeightLayerNames = InSectionMesh.GetWeightLayerNames();
+		const TArray<FName> MeshPartitionChannels = InChannels.GetChannels();
 		TArray<int32> SectionChannelIndex;
-		// #todo: when the dynamic mesh path is deprecated, these can simply be views over the raw attribute channel data.
-		// This currently has to be a local copy because data stored in the overlays can't be retrieved as a simple contiguous array view.
 		TArray<TArray<float>> SourceSectionChannels;
 		int NumUsedChannels = 0;
-		const TArray<FName> WeightLayerNames = InSectionMesh.GetWeightLayerNames();
-		for (int i = 0; i < NumChannelsInAttributes; ++i)
+		int32 ChannelIndex = 0;
+		for (const FName ChannelName : MeshPartitionChannels)
 		{
-			const FName WeightLayerName = WeightLayerNames[i];
-			TArray<float> WeightLayer = InSectionMesh.GetWeightLayerValues(WeightLayerName);
-
-			// Grab the Definition's ChannelMap or if none assigned make an empty one
-			int32 ChannelID = InChannels.FindChannel(WeightLayerName);
-			if (ChannelID != INDEX_NONE)
+			if (WeightLayerNames.Find(ChannelName) != INDEX_NONE)
 			{
-				SourceSectionChannels.Emplace(MoveTemp(WeightLayer));
-				SectionChannelIndex.Add(ChannelID);
-				NumUsedChannels++;
+				TArray<float> WeightLayer = InSectionMesh.GetWeightLayerValues(ChannelName);
 
-				// TODO: SG : Currently bounding the number of channels so it fit in a single channel map
-				if (NumUsedChannels == MaxNumChannels)
+				// Check that the values in this layer are valid and active (not null):
+				int32 ActiveCount = 0;
+				for (float Weight : WeightLayer)
 				{
-					break;
+					if (Weight > 0.0f)
+					{
+						++ActiveCount;
+					}
 				}
-			}
-		}
 
-		for (int32 VID = 0; VID < NumVertices; ++VID)
-		{
-			int32 ChannelOffset = VID * MaxNumChannels;
-
-			for (int C = 0; C < NumUsedChannels; ++C)
-			{
-				const float Value = SourceSectionChannels[C][VID];
-				if (Value != 0)
+				// Channel is active
+				if (ActiveCount && WeightLayer.Num() == NumVertices)
 				{
-					int32 MaskIndex = FMath::Min(MaxNumChannels - 1, SectionChannelIndex[C]);
-					SectionChannels[ChannelOffset + MaskIndex] = Value;
-					ChannelActiveCounts[MaskIndex]++;
+					SourceSectionChannels.Emplace(MoveTemp(WeightLayer));
+					SectionChannelIndex.Add(ChannelIndex);
+					ChannelActiveCounts[ChannelIndex] = ActiveCount;
 				}
 			}
+
+			ChannelIndex++;
 		}
-	}
 
-	InSection.UVmesh.ChannelActiveCounts = ChannelActiveCounts;
-	InSection.UVmesh.Channels = SectionChannels;
+		// gather the section channels values
+		// and allocate new channel entry if needed
+		InSection.UVmesh.Channels.Init(0, NumVertices * SourceSectionChannels.Num());
 
+		uint64 ChannelWeightsOffset = NumVertices * sizeof(float);
+		for (int32 i = 0; i < SourceSectionChannels.Num(); ++i)
+		{
+			FMemory::Memcpy((InSection.UVmesh.Channels.GetData() + i * NumVertices), SourceSectionChannels[i].GetData(), NumVertices * sizeof(float));
+		}
+	}
+	InSection.UVmesh.ChannelActiveCounts = ChannelActiveCounts;
 
 	// Gather the actual active channels with non zero values
 	// Populate the channel offset table
 	int8 ActiveChannelIndex = 0;
 	TArray<uint8> ActiveChannelIndices;
 	TArray<uint8> ChannelTableIndices;
-
 	for (int32 C : ChannelActiveCounts)
 	{
 		if (C > 0)
@@ -352,6 +357,7 @@ void FChannelTextureRenderer::RenderSectionChannels(UTexture* InOutSectionTextur
 			
 			int32 NumIndices = Section.UVmesh.NumIndices();
 			int32 NumVertices = Section.UVmesh.NumVertices();
+			int32 NumChannels = Section.UVmesh.NumChannels();
 
 			TArray<uint8> ChannelOffsetTable = Section.ChannelOffsetTable;
 
@@ -387,10 +393,9 @@ void FChannelTextureRenderer::RenderSectionChannels(UTexture* InOutSectionTextur
 				RDGSectionMeshUVsSRV = GraphBuilder.CreateSRV(RDGSectionBufferDesc);
 			}
 
-			// Section ChannelMasks
-			FRDGBufferSRVRef RDGSectionMeshChannelMasksSRV;
+			// Section Channels
+			FRDGBufferSRVRef RDGSectionMeshChannelsSRV;
 			{
-				int32 NumChannels = Section.UVmesh.NumChannels();
 				// Section turned into RDG Resource
 				FRDGBufferRef RDGSectionBuffer = GraphBuilder.CreateBuffer(FRDGBufferDesc::CreateBufferDesc(sizeof(float), NumChannels * NumVertices), TEXT("SectionChannelMasks"));
 				GraphBuilder.QueueBufferUpload(RDGSectionBuffer, Section.UVmesh.Channels.GetData(), sizeof(float) * NumChannels * NumVertices);
@@ -399,14 +404,14 @@ void FChannelTextureRenderer::RenderSectionChannels(UTexture* InOutSectionTextur
 				RDGSectionBufferDesc.NumElements = NumChannels * NumVertices;
 				RDGSectionBufferDesc.StartOffsetBytes = 0;
 
-				RDGSectionMeshChannelMasksSRV = GraphBuilder.CreateSRV(RDGSectionBufferDesc);
+				RDGSectionMeshChannelsSRV = GraphBuilder.CreateSRV(RDGSectionBufferDesc);
 			}
 
 			// Pass Shaders
 			FGlobalShaderMap* ShaderMap = GetGlobalShaderMap(GMaxRHIFeatureLevel);
 			// Shader with Permutation : auto VertexShader = GlobalShaderMap->GetShader<FRasterizeToRectsVS>(PermutationVector);
-			auto VertexShader = ShaderMap->GetShader<FMegaMesh_DrawUVDomainVS>();
-			auto PixelShader = ShaderMap->GetShader<FMegaMesh_DrawUVDomainPS>();
+			auto VertexShader = ShaderMap->GetShader<FMeshPartition_DrawUVDomainVS>();
+			auto PixelShader = ShaderMap->GetShader<FMeshPartition_DrawUVDomainPS>();
 
 			// Destination texture to RDG
 			FTextureResource* DestResource = DestSectionTexture->GetResource();
@@ -425,20 +430,21 @@ void FChannelTextureRenderer::RenderSectionChannels(UTexture* InOutSectionTextur
 
 			FRDGTextureRef RDGMaskTexture = GraphBuilder.CreateTexture(FRDGTextureDesc::Create2D(
 				DestTexture->GetDesc().Extent,
-				PF_R8_UINT,
+				PF_R8,
 				FClearValueBinding::Black,
-				ETextureCreateFlags::ShaderResource | ETextureCreateFlags::RenderTargetable,
-				1 // Single mip for the mask
+				ETextureCreateFlags::ShaderResource | ETextureCreateFlags::UAV | ETextureCreateFlags::RenderTargetable,
+				DestTexture->GetDesc().NumMips
 			), TEXT("MeshPartition_SectionMask_PassRenderTarget"));
 
 			for (int32 Page = 0; Page < Section.NumActivePages(); ++Page)
 			{
 				{
 					// Pass Parameters
-					FMegaMesh_DrawUVDomain_Parameters* PassParameters = GraphBuilder.AllocParameters<FMegaMesh_DrawUVDomain_Parameters>();
+					FMeshPartition_DrawUVDomain_Parameters* PassParameters = GraphBuilder.AllocParameters<FMeshPartition_DrawUVDomain_Parameters>();
 					PassParameters->VS.InMeshIndices = RDGSectionMeshIndicesSRV;
 					PassParameters->VS.InMeshUVs = RDGSectionMeshUVsSRV;
-					PassParameters->VS.InMeshChannelMasks = RDGSectionMeshChannelMasksSRV;
+					PassParameters->VS.InMeshChannels = RDGSectionMeshChannelsSRV;
+					PassParameters->VS.InChannelOffset = Page * NumVertices; // Channel begining offset in the Channels buffer
 
 					ClearUnusedGraphResources(VertexShader, &PassParameters->VS);
 					ClearUnusedGraphResources(PixelShader, &PassParameters->PS);
@@ -474,30 +480,12 @@ void FChannelTextureRenderer::RenderSectionChannels(UTexture* InOutSectionTextur
 							FIntPoint ViewportSize(Section.RasterResolution.X, Section.RasterResolution.Y);
 							RHICmdList.SetViewport(0.0f, 0.0f, 0.0f, (float)ViewportSize.X, (float)ViewportSize.Y, 1.0f);
 
-							int32 ActiveChannelIndex = Page * NumChannelsPerPage;
-							FUintVector4 PageChannelsOffset(0);
-							FVector4f PageChannelsMask(0.0f);
-							for (int i = 0; i < NumChannelsPerPage; ++i)
-							{
-								if (ActiveChannelIndex + i < Section.UVmesh.ActiveChannelIndices.Num())
-								{
-									PageChannelsOffset[i] = Section.UVmesh.ActiveChannelIndices[ActiveChannelIndex + i];
-									PageChannelsMask[i] = 1.0f;
-								}
-							}
-
-							int32 NumChannels = Section.UVmesh.NumChannels();
-							int32 NumChannelsLeft = NumChannels - Page * NumChannelsPerPage;
-							FVector4f PageChannelMask(1.0f, 1 < NumChannelsLeft, 2 < NumChannelsLeft, 3 < NumChannelsLeft);
-
 							for (int32 PassId = 0; PassId < Section.Bases_UVBox.Num(); ++PassId)
 							{
 								const FInt32Vector2 BaseMeshRange = Section.UVmesh.BaseIndicesRanges[PassId];
 
-								PassParameters->VS.InDrawCall = FUintVector4(Page, NumChannels, BaseMeshRange.X, BaseMeshRange.Y);
-								PassParameters->VS.InChannelsOffset = PageChannelsOffset;
-								PassParameters->VS.InChannelsMask = PageChannelsMask;
-
+								PassParameters->VS.InDrawCall = FUintVector4(Section.UVmesh.NumVertices(), Section.UVmesh.NumChannels(), BaseMeshRange.X, BaseMeshRange.Y);
+	
 								SetShaderParameters(RHICmdList, VertexShader, VertexShader.GetVertexShader(), PassParameters->VS);
 								SetShaderParameters(RHICmdList, PixelShader, PixelShader.GetPixelShader(), PassParameters->PS);
 
@@ -506,14 +494,15 @@ void FChannelTextureRenderer::RenderSectionChannels(UTexture* InOutSectionTextur
 							}
 						});
 
-					if (CVarMeshPartition_BorderFill->GetBool())
+					// Border fill
+					if (CVarMeshPartition_BorderFillMethod->GetInt() == 1)
 					{
-						FMegaMesh_BorderFillCS::FParameters* BorderParams = GraphBuilder.AllocParameters<FMegaMesh_BorderFillCS::FParameters>();
+						FMeshPartition_BorderFillCS::FParameters* BorderParams = GraphBuilder.AllocParameters<FMeshPartition_BorderFillCS::FParameters>();
 						BorderParams->Mask = GraphBuilder.CreateSRV(RDGMaskTexture);
 						BorderParams->RWSectionTexture = GraphBuilder.CreateUAV(RDGRenderTarget);
 						BorderParams->Resolution = FUintVector2(RDGRenderTarget->Desc.Extent.X, RDGRenderTarget->Desc.Extent.Y);
 
-						TShaderRef<FMegaMesh_BorderFillCS> ComputeShader = ShaderMap->GetShader<FMegaMesh_BorderFillCS>();
+						TShaderRef<FMeshPartition_BorderFillCS> ComputeShader = ShaderMap->GetShader<FMeshPartition_BorderFillCS>();
 						
 						FComputeShaderUtils::AddPass(
 							GraphBuilder,
@@ -526,9 +515,93 @@ void FChannelTextureRenderer::RenderSectionChannels(UTexture* InOutSectionTextur
 								FIntVector(8, 8, 1)
 							)
 						);
+
+						// Regular border fill and mip generation
+						FGenerateMips::ExecuteCompute(GraphBuilder, ERHIFeatureLevel::SM6, RDGRenderTarget, TStaticSamplerState<SF_Bilinear, AM_Clamp, AM_Clamp, AM_Clamp>::GetRHI());
+					}
+					// Push pull
+					else if (CVarMeshPartition_BorderFillMethod->GetInt() == 2)
+					{
+						const FRDGTextureDesc& TextureDesc = RDGRenderTarget->Desc;						
+						FRDGTextureSRVDesc SRVDesc(RDGRenderTarget);
+						FRDGTextureUAVDesc UAVDesc(RDGRenderTarget);
+						FRDGTextureSRVDesc SRVMaskDesc(RDGMaskTexture);
+						FRDGTextureUAVDesc UAVMaskDesc(RDGMaskTexture);
+						SRVDesc.NumMipLevels = 1;
+						SRVMaskDesc.NumMipLevels = 1;
+
+						{
+							TShaderRef<FMeshPartition_FillPullCS> PullShader = ShaderMap->GetShader<FMeshPartition_FillPullCS>();
+
+							// Loop through each level of the mips that require creation and add a dispatch pass per level.
+							for (int8 MipLevel = 1; MipLevel < TextureDesc.NumMips; ++MipLevel)
+							{
+								const FIntPoint DestTextureSize(
+									FMath::Max(TextureDesc.Extent.X >> MipLevel, 1),
+									FMath::Max(TextureDesc.Extent.Y >> MipLevel, 1));
+
+								SRVDesc.MipLevel = (int8)(MipLevel - 1);
+								SRVMaskDesc.MipLevel = (int8)(MipLevel - 1);
+								UAVDesc.MipLevel = (int8)(MipLevel);
+								UAVMaskDesc.MipLevel = (int8)(MipLevel);
+
+								FMeshPartition_FillPullCS::FParameters* ShaderParams = GraphBuilder.AllocParameters<FMeshPartition_FillPullCS::FParameters>();
+								ShaderParams->TexelSize = FVector2f(1.0f / DestTextureSize.X, 1.0f / DestTextureSize.Y);
+								ShaderParams->SectionMipIn = GraphBuilder.CreateSRV(SRVDesc);
+								ShaderParams->SectionMipOut = GraphBuilder.CreateUAV(UAVDesc);
+								ShaderParams->MaskMipIn = GraphBuilder.CreateSRV(SRVMaskDesc);
+								ShaderParams->MaskMipOut = GraphBuilder.CreateUAV(UAVMaskDesc);
+								ShaderParams->MipSampler = TStaticSamplerState<SF_Bilinear, AM_Clamp, AM_Clamp, AM_Clamp>::GetRHI();
+								ShaderParams->ResolutionPass = FIntVector4(DestTextureSize.X, DestTextureSize.Y, MipLevel, TextureDesc.NumMips);
+
+								FComputeShaderUtils::AddPass(
+									GraphBuilder,
+									RDG_EVENT_NAME("MeshPartition_FillPull DestMipLevel=%d", MipLevel),
+									ERDGPassFlags::Compute,
+									PullShader,
+									ShaderParams,
+									FComputeShaderUtils::GetGroupCount(DestTextureSize, FComputeShaderUtils::kGolden2DGroupSize)
+								);
+							}
+						}
+						{
+							TShaderRef<FMeshPartition_FillPushCS> PushShader = ShaderMap->GetShader<FMeshPartition_FillPushCS>();
+
+							// Loop through each level of the mips that require creation and add a dispatch pass per level.
+							for (int8 MipLevel = TextureDesc.NumMips - 2; MipLevel >= 0; --MipLevel)
+							{
+								const FIntPoint DestTextureSize(
+									FMath::Max(TextureDesc.Extent.X >> MipLevel, 1),
+									FMath::Max(TextureDesc.Extent.Y >> MipLevel, 1));
+
+								SRVDesc.MipLevel = (int8)(MipLevel + 1);
+								UAVDesc.MipLevel = (int8)(MipLevel);
+								SRVMaskDesc.MipLevel = (int8)(MipLevel);
+
+								FMeshPartition_FillPushCS::FParameters* ShaderParams = GraphBuilder.AllocParameters<FMeshPartition_FillPushCS::FParameters>();
+								ShaderParams->TexelSize = FVector2f(1.0f / DestTextureSize.X, 1.0f / DestTextureSize.Y);
+								ShaderParams->SectionMipIn = GraphBuilder.CreateSRV(SRVDesc);
+								ShaderParams->SectionMipOut = GraphBuilder.CreateUAV(UAVDesc);
+								ShaderParams->MaskMipIn = GraphBuilder.CreateSRV(SRVMaskDesc);
+								ShaderParams->MipSampler = TStaticSamplerState<SF_Bilinear, AM_Clamp, AM_Clamp, AM_Clamp>::GetRHI();
+								ShaderParams->ResolutionPass = FIntVector4(DestTextureSize.X, DestTextureSize.Y, MipLevel, TextureDesc.NumMips);
+
+								FComputeShaderUtils::AddPass(
+									GraphBuilder,
+									RDG_EVENT_NAME("MeshPartition_FillPush DestMipLevel=%d", MipLevel),
+									ERDGPassFlags::Compute,
+									PushShader,
+									ShaderParams,
+									FComputeShaderUtils::GetGroupCount(DestTextureSize, FComputeShaderUtils::kGolden2DGroupSize)
+								);
+							}
+						}
+					}
+					else
+					{
+						// No border fill but just mip generation
+						FGenerateMips::ExecuteCompute(GraphBuilder, ERHIFeatureLevel::SM6, RDGRenderTarget, TStaticSamplerState<SF_Bilinear, AM_Clamp, AM_Clamp, AM_Clamp>::GetRHI());
 					}
-					
-					FGenerateMips::ExecuteCompute(GraphBuilder, ERHIFeatureLevel::SM6, RDGRenderTarget, TStaticSamplerState<SF_Bilinear, AM_Clamp, AM_Clamp, AM_Clamp>::GetRHI());
 
 					// Make sure the texture is actually managed by RDG as a call to RegisterExternalTexture does not guarantee this if the texture was already registered and then had UseExternalAccessMode() called on it.
 					GraphBuilder.UseInternalAccessMode(RDGDestTexture);
diff --git a/Engine/Plugins/Experimental/MeshPartition/Source/MeshPartitionEditor/Private/MeshPartitionChannelCollection_GenerateChannelUV.cpp b/Engine/Plugins/Experimental/MeshPartition/Source/MeshPartitionEditor/Private/MeshPartitionChannelCollection_GenerateChannelUV.cpp
--- a/Engine/Plugins/Experimental/MeshPartition/Source/MeshPartitionEditor/Private/MeshPartitionChannelCollection_GenerateChannelUV.cpp
+++ b/Engine/Plugins/Experimental/MeshPartition/Source/MeshPartitionEditor/Private/MeshPartitionChannelCollection_GenerateChannelUV.cpp
@@ -358,7 +358,7 @@ namespace MegaMeshChannelCollectionLocals
 		}
 	}
 
-	void CreateUVLayout_SetTriangleUVsThroughDynamicMeshTools(FMeshData& InOutSectionMesh, TArray<TArray<int32>>& OutIslandTriangles)
+	void CreateUVLayout_SetTriangleUVsThroughDynamicMeshTools(FSectionDomainMapping& InSectionDomainMapping, FMeshData& InOutSectionMesh, TArray<TArray<int32>>& OutIslandTriangles)
 	{
 		// Before packing UVs, compute from-scratch UV layouts if needed
 		{
@@ -378,7 +378,7 @@ namespace MegaMeshChannelCollectionLocals
 				int32 UseBaseIndex = 0;
 
 				CreateUVLayout_PerformThroughDynamicMesh(InOutSectionMesh, ToUnwrapTIDs,
-					[](Geometry::FDynamicMesh3& Mesh)
+														 [InSectionDomainMapping](Geometry::FDynamicMesh3& Mesh)
 					{
 						{
 							// Create temp array of all triangle IDs, for use by UV layout methods below
@@ -392,7 +392,7 @@ namespace MegaMeshChannelCollectionLocals
 							// Quick initial UV layout method, TODO: revisit algorithm choice and scale
 							Geometry::FDynamicMeshUVEditor UVEditor(&Mesh, 0, true);
 							UVEditor.SetTriangleUVsFromBoxProjection(AllMeshTIDs, [](const FVector3d& V) { return V; }, Geometry::FFrame3d(), FVector3d(100, 100, 100));
-							UVEditor.QuickPack();
+							UVEditor.QuickPack(InSectionDomainMapping.MinImageResolution, 2 * InSectionDomainMapping.GetGutterTexelCount());
 							UVEditor.ScaleUVAreaTo3DArea(AllMeshTIDs, true, .00001);
 						}
 					},
@@ -484,7 +484,8 @@ namespace MegaMeshChannelCollectionLocals
 			double MajorAxisSign = FMathd::Sign(ScaledNormal[MajorAxis]);
 			int32 Bucket = (MajorAxisSign > 0) ? (MajorAxis + 3) : MajorAxis;
 			int32 MappedBaseID = GetMappedMeshBaseID(tid);
-			int BucketID = MappedBaseID * 6 + Bucket;
+			//int BucketID = MappedBaseID * 6 + Bucket;
+			int BucketID = Bucket;
 			TriangleBoxPlaneAssignments[i] = FIndex2i(MajorAxis, BucketID);
 			IndexMap[tid] = i;
 
@@ -744,7 +745,41 @@ namespace MegaMeshChannelCollectionLocals
 		OutIslandTriangles = MoveTemp(DisjointedIslandTriangles);
 	}
 
-	void CreateUVLayout_EvaluateUVDomainDimensions(double& OutUVDomainSize, int32& OutImageResolution, const FMeshData& InSectionMesh, float InTexelSize, int32 InMaxResolution)
+	void CreateUVLayout_Evaluate3DDomainDimensions(FSectionDomainMapping& OutSectionDomainMapping, const FMeshData& InSectionMesh)
+	{
+		TRACE_CPUPROFILER_EVENT_SCOPE(MeshPartition::CreateUVLayout_Evaluate3DDomainDimensions);
+
+		// Measure the true 3d area of the mesh
+		double Section_Area3D = 0;
+		for (int32 TriangleId : InSectionMesh.TriangleIndicesItr())
+		{
+			Geometry::FIndex3i Triangle = InSectionMesh.GetTriangle(TriangleId);
+
+			// Compute the triangle area and add up to the base's area
+			FVector3d P0 = InSectionMesh.GetVertex(Triangle.A);
+			FVector3d P1 = InSectionMesh.GetVertex(Triangle.B);
+			FVector3d P2 = InSectionMesh.GetVertex(Triangle.C);
+			double TriangleArea3D = 0.5 * FVector3d::CrossProduct((P1 - P0), (P2 - P0)).Length();
+			Section_Area3D += TriangleArea3D;
+		}
+
+		// Project the 3d area corresponding to a fully filled uv domain
+		// ignore the gutter area
+		// And estimate the size of that "fully filled" 3d square
+		double SectionSize = FMath::Sqrt(Section_Area3D);
+
+		// From the required Texel size (in 3d unit) evaluate the ideal image resolution
+		OutSectionDomainMapping.MinImageResolution = SectionSize / OutSectionDomainMapping.TexelSize3D; // Concrete image size per uv space unit [texel/uv]
+		OutSectionDomainMapping.ImageResolution = OutSectionDomainMapping.MinImageResolution;
+
+		OutSectionDomainMapping.Area3D = Section_Area3D;
+		OutSectionDomainMapping.AreaUV = 1.0;
+								
+		OutSectionDomainMapping.Size3D = SectionSize;
+		OutSectionDomainMapping.SizeUV = 1.0;
+	}
+
+	void CreateUVLayout_EvaluateUVDomainDimensions(FSectionDomainMapping& InOutSectionDomainMapping, const FMeshData& InSectionMesh)
 	{
 		TRACE_CPUPROFILER_EVENT_SCOPE(MeshPartition::CreateUVLayout_EvaluateUVDomainDimensions);
 
@@ -778,22 +813,34 @@ namespace MegaMeshChannelCollectionLocals
 		double SectionSize = FMath::Sqrt(ProjectedSectionDomainArea3D);
 
 		// From the required Texel size (in 3d unit) evaluate the ideal image resolution
-		double TexelSize_Unit = InTexelSize; // Size of one texel in position unit [u] assigned from definition
-		double IdealImageResolution = SectionSize / TexelSize_Unit; // Concrete image size per uv space unit [texel/uv]
+		double TexelSize_Unit = InOutSectionDomainMapping.TexelSize3D; // Size of one texel in position unit [u] assigned from definition
+		double UVScaledImageResolution = SectionSize / TexelSize_Unit; // Concrete image size per uv space unit [texel/uv]
+
+		double IdealImageResolution = FMath::Max(InOutSectionDomainMapping.MinImageResolution, UVScaledImageResolution);
 
 		// True image resolution needs to be a power of 2, pick the nearest size up
 		int32 ImageResolution = FMath::RoundUpToPowerOfTwo((int32)FMath::Floor(IdealImageResolution));
-		ImageResolution = FMath::Min(InMaxResolution, ImageResolution);
+		ImageResolution = FMath::Min(InOutSectionDomainMapping.MaxImageResolution, ImageResolution);
+
 
-		OutUVDomainSize = SectionSize;
-		OutImageResolution = ImageResolution;
+		InOutSectionDomainMapping.Size3D = SectionSize;
+		InOutSectionDomainMapping.SizeUV = SectionSize;
+
+		InOutSectionDomainMapping.Area3D = Section_Area3D;
+		InOutSectionDomainMapping.AreaUV = Section_AreaUV;
+
+		InOutSectionDomainMapping.ImageResolution = ImageResolution;
 	}
 
 
-	void CreateUVLayout_PackIslandsNormalizeUVDomain(FMeshData& InOutSectionMesh, const TArray<TArray<int32>>& InIslandTriangles, int32 InTextureResolution )
+	void CreateUVLayout_PackIslandsNormalizeUVDomain(FSectionDomainMapping& OutSectionDomainMapping,  FMeshData& InOutSectionMesh, const TArray<TArray<int32>>& InIslandTriangles)
 	{
 		TRACE_CPUPROFILER_EVENT_SCOPE(MeshPartition::CreateUVLayout_PackIslandsNormalizeUVDomain);
 	
+		MegaMeshChannelCollectionLocals::CreateUVLayout_EvaluateUVDomainDimensions(OutSectionDomainMapping, InOutSectionMesh);
+
+
+
 		int32 NumVertices = InOutSectionMesh.MaxVertexID();
 
 		// Each vertex belongs to one island which should match all the tringle's BaseID referencing it
@@ -846,7 +893,7 @@ namespace MegaMeshChannelCollectionLocals
 
 		// then run the atlas algortihm
 		// Find the new UV domain for the section
-		float IslandMargin = 4.0 / (float) InTextureResolution;
+		float IslandMargin = OutSectionDomainMapping.GetGutterUV();
 		TArray<FBox2f> FreeRects;
 		TArray<FBox2f> SectionIslandUVDomains = MegaMeshChannelCollectionLocals::Packer::CreateAtlas(IslandUVBoxes, IslandMargin, FreeRects);
 
@@ -888,25 +935,27 @@ void FChannelTextureRenderer::CreateSectionUVLayout(FMeshData& InOutSectionMesh,
 {
 	TRACE_CPUPROFILER_EVENT_SCOPE(MeshPartition::FChannelTextureRenderer::CreateSectionUVLayout);
 
+	
+	FSectionDomainMapping DomainMapping{
+		.TexelSize3D = GenerationOptions.TexelSize,
+		.MaxImageResolution = GenerationOptions.MaxTextureResolution
+	};
+	MegaMeshChannelCollectionLocals::CreateUVLayout_Evaluate3DDomainDimensions(DomainMapping, InOutSectionMesh);
+
 	// Unwrap is currently default-disabled until it can be moved to a background thread
 	const bool bPerformUnwrapThroughDynamicMeshTools = GenerationOptions.UVLayoutMethod == EChannelCollectionUVLayoutMethod::ReferenceBoxProject;
 	TArray<TArray<int32>> IslandTriangles;
 
 	if (bPerformUnwrapThroughDynamicMeshTools)
 	{
-		MegaMeshChannelCollectionLocals::CreateUVLayout_SetTriangleUVsThroughDynamicMeshTools(InOutSectionMesh, IslandTriangles);
+		MegaMeshChannelCollectionLocals::CreateUVLayout_SetTriangleUVsThroughDynamicMeshTools(DomainMapping, InOutSectionMesh, IslandTriangles);
 	}
 	else
 	{
 		MegaMeshChannelCollectionLocals::CreateUVLayout_SetTriangleUVsFromBoxProjection(InOutSectionMesh, IslandTriangles);
 	}
 
-	double OutUVDomainSize = 0;
-	int32 OutImageResolution = 0;
-	MegaMeshChannelCollectionLocals::CreateUVLayout_EvaluateUVDomainDimensions(OutUVDomainSize, OutImageResolution, InOutSectionMesh, GenerationOptions.TexelSize, GenerationOptions.MaxTextureResolution);
-
-
-	MegaMeshChannelCollectionLocals::CreateUVLayout_PackIslandsNormalizeUVDomain(InOutSectionMesh, IslandTriangles, GenerationOptions.MaxTextureResolution);
+	MegaMeshChannelCollectionLocals::CreateUVLayout_PackIslandsNormalizeUVDomain(DomainMapping, InOutSectionMesh, IslandTriangles);
 	// UVs have been assigned and normalized to section mesh
 }
 
@@ -939,9 +988,20 @@ void FChannelTextureRenderer::GenerateSectionMeshUVs(const FMeshData& InSectionM
 		});
 
 
+	int32 NumTriangles = SectionIndices.Num() / 3;
+	TArray<int32> SectionIslandPerVertex;
+	SectionIslandPerVertex.Init(-1, NumVertices);
+	TArray<int32> SectionIslandPerTriangles;
+	SectionIslandPerTriangles.Init(-1, NumTriangles);
+
 	double OutUVDomainSize = 0;
 	int32 OutImageResolution = 0;
-	MegaMeshChannelCollectionLocals::CreateUVLayout_EvaluateUVDomainDimensions(OutUVDomainSize, OutImageResolution, InSectionMesh, InTexelSize, DefaultMaxImageResolution);
+
+	FSectionDomainMapping InOutDomainMapping{
+		.TexelSize3D = InTexelSize,
+		.MaxImageResolution = DefaultMaxImageResolution
+	};
+	MegaMeshChannelCollectionLocals::CreateUVLayout_EvaluateUVDomainDimensions(InOutDomainMapping, InSectionMesh);
 
 	FBox2f UVDomain({ 0,0 }, { 1,1 });
 	InSection.Islands_Bases = { {0} };
@@ -953,7 +1013,9 @@ void FChannelTextureRenderer::GenerateSectionMeshUVs(const FMeshData& InSectionM
 	InSection.UVmesh.Indices = MoveTemp(SectionIndices);
 	InSection.UVmesh.UVs = MoveTemp(SectionUVs);
 
-	InSection.RasterResolution = FUint32Vector2((uint32)OutImageResolution);
-	InSection.TexcoordMetrics = FVector2f(OutUVDomainSize);
+	InSection.RasterResolution = FUint32Vector2((uint32)InOutDomainMapping.ImageResolution);
+	InSection.TexcoordMetrics = FVector2f(InOutDomainMapping.Size3D);
+	InSection.DomainMapping = InOutDomainMapping;
 }
+
 } // namespace UE::MeshPartition
diff --git a/Engine/Plugins/Experimental/MeshPartition/Source/MeshPartitionEditor/Private/MeshPartitionEditorComponent.cpp b/Engine/Plugins/Experimental/MeshPartition/Source/MeshPartitionEditor/Private/MeshPartitionEditorComponent.cpp
--- a/Engine/Plugins/Experimental/MeshPartition/Source/MeshPartitionEditor/Private/MeshPartitionEditorComponent.cpp
+++ b/Engine/Plugins/Experimental/MeshPartition/Source/MeshPartitionEditor/Private/MeshPartitionEditorComponent.cpp
@@ -381,6 +381,7 @@ void UMeshPartitionEditorComponent::BuildMegaMeshPreviewSections(TConstArrayView
 	Settings.bCacheResult = true;
 	Settings.bAllowDDCWrite = CVarEnablePreviewSectionDDCWrite.GetValueOnGameThread();
 	Settings.bAllowDDCRead = CVarEnablePreviewSectionDDCRead.GetValueOnGameThread();
+
 	Settings.TexcoordGenerationOptions = FChannelCollectionUVLayoutOptions::GetFromDefinition(GetMegaMeshDefinition());
 
 	MeshPartition::FBuilderSettings::FChannelRenderSettings ChannelRenderSettings;
@@ -1600,6 +1601,7 @@ void UMeshPartitionEditorComponent::FinalizePreviewSectionBuilds()
 			PreviewSection->SetChannelTexture(SectionChannel->Texture.Get());
 			PreviewSection->SetChannelTexcoordDesc(SectionChannel->TexcoordMetrics);
 			PreviewSection->SetChannelTable(SectionChannel->Table);
+			PreviewSection->SetChannelGenerationMesh(SectionChannel->Mesh);
 
 			PreviewSections.Emplace(PreviewSection);
 
diff --git a/Engine/Plugins/Experimental/MeshPartition/Source/MeshPartitionEditor/Private/MeshPartitionMaterialExpressionUtils.cpp b/Engine/Plugins/Experimental/MeshPartition/Source/MeshPartitionEditor/Private/MeshPartitionMaterialExpressionUtils.cpp
--- a/Engine/Plugins/Experimental/MeshPartition/Source/MeshPartitionEditor/Private/MeshPartitionMaterialExpressionUtils.cpp
+++ b/Engine/Plugins/Experimental/MeshPartition/Source/MeshPartitionEditor/Private/MeshPartitionMaterialExpressionUtils.cpp
@@ -132,8 +132,10 @@ namespace MaterialExpression
 
 			// Fetch texel at the correct slice and 2d uv
 			int32 TextureFromCollectionCodeIndex = CompilePickChannelTexture(Compiler, InTextureCollection_IDX);
-			int32 UVCoord_IDX = Compiler->AppendVector(CompileChannelTexcoordUV(Compiler), ChannelTextureSlot_IDX);
-			int32 FetchedValueChannel_IDX = Compiler->TextureSample(TextureFromCollectionCodeIndex, UVCoord_IDX, SAMPLERTYPE_Masks, -1, -1, TMVM_None, SSM_FromTextureAsset);
+			int32 UV_IDX = CompileChannelTexcoordUV(Compiler);
+			int32 UVSlice_IDX = Compiler->AppendVector(UV_IDX, ChannelTextureSlot_IDX);
+			int32 UV_ConstSlice_IDX = Compiler->AppendVector(UV_IDX, Compiler->Constant(0)); // pass a pure uv and constant slice for derivative evaluation to avoid nanite uv seams issues
+			int32 FetchedValueChannel_IDX = Compiler->TextureSample(TextureFromCollectionCodeIndex, UVSlice_IDX, SAMPLERTYPE_Masks, Compiler->DDX(UV_ConstSlice_IDX), Compiler->DDY(UV_ConstSlice_IDX), TMVM_Derivative, SSM_FromTextureAsset);
 
 			// If the channel texture slot returned is -1, it means the channel is empty in the section, exit returning 0
 			// else return the fetched channel value
diff --git a/Engine/Plugins/Experimental/MeshPartition/Source/MeshPartitionEditor/Private/MeshPartitionRenderingUtils.cpp b/Engine/Plugins/Experimental/MeshPartition/Source/MeshPartitionEditor/Private/MeshPartitionRenderingUtils.cpp
--- a/Engine/Plugins/Experimental/MeshPartition/Source/MeshPartitionEditor/Private/MeshPartitionRenderingUtils.cpp
+++ b/Engine/Plugins/Experimental/MeshPartition/Source/MeshPartitionEditor/Private/MeshPartitionRenderingUtils.cpp
@@ -14,7 +14,7 @@ FRenderCaptureManager& FRenderCaptureManager::Get()
 {
 	static FRenderCaptureManager Instance;
 	static FAutoConsoleCommand CCmdRenderDocCaptureFrame = FAutoConsoleCommand(
-		TEXT("MegaMesh.CaptureNextBatch"),
+		TEXT("MeshPartition.Channels.CaptureNextBatch"),
 		TEXT("Captures the next Job Batch and launches RenderDoc"),
 		FConsoleCommandDelegate::CreateLambda([]()
 			{
@@ -33,7 +33,7 @@ void FRenderCaptureManager::BeginCapture(FRDGBuilder& InGraphBuilder)
 {
 	if (FRenderCaptureManager::Get().NextCaptureCounter > 0)
 	{
-		FRenderCaptureManager::Get().CurrentCapture = MakeUnique<RenderCaptureInterface::FScopedCapture>(true, InGraphBuilder, *FString("MegaMesh"));
+		FRenderCaptureManager::Get().CurrentCapture = MakeUnique<RenderCaptureInterface::FScopedCapture>(true, InGraphBuilder, *FString("MeshPartition"));
 		FRenderCaptureManager::Get().NextCaptureCounter = 0;
 	}
 }
diff --git a/Engine/Plugins/Experimental/MeshPartition/Source/MeshPartitionEditor/Public/MeshPartitionChannelCollection.h b/Engine/Plugins/Experimental/MeshPartition/Source/MeshPartitionEditor/Public/MeshPartitionChannelCollection.h
--- a/Engine/Plugins/Experimental/MeshPartition/Source/MeshPartitionEditor/Public/MeshPartitionChannelCollection.h
+++ b/Engine/Plugins/Experimental/MeshPartition/Source/MeshPartitionEditor/Public/MeshPartitionChannelCollection.h
@@ -40,11 +40,29 @@ struct FChannelCollectionUVLayoutOptions
 	EChannelCollectionUVLayoutMethod UVLayoutMethod;
 };
 
+struct FSectionDomainMapping
+{
+	double	TexelSize3D = 100.0; // Size of one texel in position unit [u] assigned from definition
+	int32	MaxImageResolution = 4092;
+	int32	MinImageResolution = 0;
+	int32	ImageResolution = 0;
+
+	double	Area3D = 0;
+	double  AreaUV = 0;
+			
+	double	Size3D = 0;
+	double  SizeUV = 0;
+
+	float GetGutterTexelCount() const { return 4.0; }
+	float GetGutterUV() const { return GetGutterTexelCount() / (float) ImageResolution; }
+};
+
 struct FSectionChannels
 {
 	TArray<uint8> Table;
 	FVector2f TexcoordMetrics = { 1.0f, 1.0f };
 	TStrongObjectPtr<UTexture> Texture;
+	FSectionChannelGenerationMeshData Mesh;
 };
 
 /**
@@ -112,6 +130,8 @@ private:
 
 		FUint32Vector2			RasterResolution;
 		FVector2f				TexcoordMetrics;
+
+		FSectionDomainMapping	DomainMapping;
 	};
 
 	static UTexture* GetDefaultTexture();
diff --git a/Engine/Plugins/Experimental/MeshPartition/Source/MeshPartitionEditor/Public/MeshPartitionPreviewSection.h b/Engine/Plugins/Experimental/MeshPartition/Source/MeshPartitionEditor/Public/MeshPartitionPreviewSection.h
--- a/Engine/Plugins/Experimental/MeshPartition/Source/MeshPartitionEditor/Public/MeshPartitionPreviewSection.h
+++ b/Engine/Plugins/Experimental/MeshPartition/Source/MeshPartitionEditor/Public/MeshPartitionPreviewSection.h
@@ -67,7 +67,12 @@ public:
 
 	UE_API void SetChannelTexcoordDesc(const FVector2f& InTexcoorDesc);
 	UE_API void SetChannelTable(const TArray<uint8>& InChannelTable);
+	UE_API const TArray<uint8>& GetChannelTable() const { return ChannelTable; }
 	UE_API void SetChannelTexture(UTexture* InChannelTexture);
+	UE_API UTexture* GetChannelTexture() const { return ChannelTexture; }
+
+	UE_API void SetChannelGenerationMesh(const MeshPartition::FSectionChannelGenerationMeshData& InSection) { SectionChannelMesh = InSection; }
+	UE_API const MeshPartition::FSectionChannelGenerationMeshData& GetChannelGenerationMesh() const { return SectionChannelMesh; }
 
 	UE_API void SetBuildPerfStats(MeshPartition::FBuildPerfStats&& InBuildPerfStats);
 	UE_API const MeshPartition::FBuildPerfStats& GetBuildPerfStats() const;
@@ -160,6 +165,9 @@ private:
 
 	FGuid GroupRegistryKey;
 
+	
+	FSectionChannelGenerationMeshData SectionChannelMesh;
+
 	friend class FMeshPreviewVisualizer;
 };
 } // namespace UE::MeshPartition
