Adding status summary functionality to the mesh partition world subsystem
Added a command to log the current status of mesh partitions in the editor world: MeshPartition.ReportSectionStatus

#Synced-CL 51480380

diff --git a/Engine/Plugins/Experimental/MeshPartition/Source/MeshPartition/Public/MeshPartitionCompiledSection.h b/Engine/Plugins/Experimental/MeshPartition/Source/MeshPartition/Public/MeshPartitionCompiledSection.h
--- a/Engine/Plugins/Experimental/MeshPartition/Source/MeshPartition/Public/MeshPartitionCompiledSection.h
+++ b/Engine/Plugins/Experimental/MeshPartition/Source/MeshPartition/Public/MeshPartitionCompiledSection.h
@@ -187,6 +187,7 @@ public:
 	UE_API void SetChannelTexcoordDesc(const FVector2f& InTexcoordDesc);
 	UE_API void SetChannelTable(const TArray<uint8>& InChannelTable);
 	UE_API void SetChannelTexture(UTexture* InChannelTexture);
+	UTexture* GetChannelTexture() { return ChannelTexture; }
 	UE_API void UpdateVirtualTextureSettings();
 	UE_API void SetupChannelDataOnChildPrimitiveComponents();
 
diff --git a/Engine/Plugins/Experimental/MeshPartition/Source/MeshPartitionEditor/Private/MeshPartitionEditorWorldSubsystem.cpp b/Engine/Plugins/Experimental/MeshPartition/Source/MeshPartitionEditor/Private/MeshPartitionEditorWorldSubsystem.cpp
--- a/Engine/Plugins/Experimental/MeshPartition/Source/MeshPartitionEditor/Private/MeshPartitionEditorWorldSubsystem.cpp
+++ b/Engine/Plugins/Experimental/MeshPartition/Source/MeshPartitionEditor/Private/MeshPartitionEditorWorldSubsystem.cpp
@@ -2,11 +2,40 @@
 
 #include "MeshPartitionEditorWorldSubsystem.h"
 #include "EngineUtils.h"
+#include "Subsystems/UnrealEditorSubsystem.h"	// UUnrealEditorSubsystem
 #include "MeshPartition.h"
 #include "MeshPartitionEditorComponent.h"
+#include "MeshPartitionWorldUpdater.h"
+#include "Editor.h"								// GEditor
 
 namespace UE::MeshPartition
 {
+static FAutoConsoleCommand ReportSectionStatus(
+	TEXT("MeshPartition.ReportSectionStatus"),
+	TEXT("Report the current mesh partition compiled section status for the current editor world"),
+	FConsoleCommandDelegate::CreateStatic(&UMeshPartitionEditorWorldSubsystem::ReportSectionStatus));
+
+void UMeshPartitionEditorWorldSubsystem::ReportSectionStatus()
+{
+	UWorld* EditorWorld = nullptr;
+	if (GEditor)
+	{
+		if (UUnrealEditorSubsystem* EditorSubsystem = GEditor->GetEditorSubsystem<UUnrealEditorSubsystem>())
+		{
+			EditorWorld = EditorSubsystem->GetEditorWorld();
+		}
+	}
+	if (EditorWorld == nullptr)
+	{
+		UE_LOG(LogMegaMeshEditor, Warning, TEXT("Dumping Compiled Section Status -- No Editor World Found"));
+		return;
+	}
+
+	UE_LOG(LogMegaMeshEditor, Warning, TEXT("Dumping Compiled Section Status for World: %s (%s)"), *EditorWorld->GetName(), *EditorWorld->GetPathName());
+	FMeshPartitionWorldUpdater MeshPartitionWorldUpdater(EditorWorld, FMeshPartitionWorldUpdater::EUpdateMode::ForCompile);
+	MeshPartitionWorldUpdater.ReportSectionStatus(EditorWorld);
+}
+
 TStatId UMeshPartitionEditorWorldSubsystem::GetStatId() const
 {
 	RETURN_QUICK_DECLARE_CYCLE_STAT(UMeshPartitionEditorWorldSubsystem, STATGROUP_Tickables);
diff --git a/Engine/Plugins/Experimental/MeshPartition/Source/MeshPartitionEditor/Private/MeshPartitionWorldUpdater.cpp b/Engine/Plugins/Experimental/MeshPartition/Source/MeshPartitionEditor/Private/MeshPartitionWorldUpdater.cpp
--- a/Engine/Plugins/Experimental/MeshPartition/Source/MeshPartitionEditor/Private/MeshPartitionWorldUpdater.cpp
+++ b/Engine/Plugins/Experimental/MeshPartition/Source/MeshPartitionEditor/Private/MeshPartitionWorldUpdater.cpp
@@ -8,6 +8,7 @@
 #include "AssetRegistry/AssetRegistryModule.h"
 #include "Engine/StaticMesh.h"
 #include "Editor/EditorEngine.h"
+#include "Engine/Texture2DArray.h"
 #include "MeshPartitionEditorUtils.h"
 
 #define LOCTEXT_NAMESPACE "FMeshPartitionWorldUpdater"
@@ -519,5 +520,382 @@ namespace UE::MeshPartition
 				MeshPartitionUpdater.AddedPlaceholderCompiledSections);
 		}
 	}
+
+	void FMeshPartitionWorldUpdater::GetStatusSummary(FStatusSummary& InOutSummary) const
+	{
+		for (const TPair<FGuid, FMeshPartitionUpdater>& Pair : MeshPartitionUpdaters)
+		{
+			const FMeshPartitionUpdater& MeshPartitionUpdater = Pair.Value;
+			MeshPartitionUpdater.GetStatusSummary(InOutSummary);
+		}
+	}
+
+	void FMeshPartitionUpdater::GetStatusSummary(FStatusSummary& InOutSummary) const
+	{
+		InOutSummary.MeshPartitionCount++;
+
+		// iterate compiled sections
+		for (const MeshPartition::FCompiledSectionStatus& Status : CompiledSectionStatus)
+		{
+			check(Status.bChecked);
+			InOutSummary.CompiledSectionCount++;
+
+			// done this way so the counts are exclusive (even though a section could fail for multiple reasons)
+			// the order here determines the priority for what is counted as the failure reason
+			if (Status.bIsPlaceholder)
+			{
+				InOutSummary.CompiledSections.Placeholders++;
+			}
+			else if (Status.bNonTargetVariant)
+			{
+				InOutSummary.CompiledSections.NonTargetVariant++;
+			}
+			else if (Status.bMismatchedModifiers)
+			{
+				InOutSummary.CompiledSections.MismatchedModifiers++;
+			}
+			else if (!Status.bPackageHashMatches)
+			{
+				InOutSummary.CompiledSections.PackageHashFails++;
+			}
+			else if (!Status.bClassHashMatches)
+			{
+				InOutSummary.CompiledSections.ClassHashFails++;
+			}
+			else if (Status.bCannotFindFile)
+			{
+				InOutSummary.CompiledSections.CannotFindFile++;
+			}
+			else if (Status.bIsDuplicate)
+			{
+				InOutSummary.CompiledSections.Duplicates++;
+			}
+			else if (Status.bReuse)
+			{
+				InOutSummary.CompiledSections.Reuse++;
+			}
+			else
+			{
+				InOutSummary.CompiledSections.Unknown++;
+			}
+		}
+
+		// also check build variant groups for any groups that do not have a compiled section
+		for (const TPair<FName, FBuildVariantUpdater>& Pair : BuildVariantUpdaters)
+		{
+			const FBuildVariantUpdater& VariantUpdater = Pair.Value;
+			InOutSummary.TargetBuildVariantCount++;
+
+			// look for any target groups with no corresponding compiled section -- these are missing
+			for (int32 ReuseCompiledSectionIndex : VariantUpdater.BestReusableCompiledSectionIndexForGroup)
+			{
+				InOutSummary.GroupCount++;
+				if (ReuseCompiledSectionIndex < 0)
+				{
+					InOutSummary.CompiledSections.Missing++;
+				}
+			}
+		}
+	}
+
+	void FMeshPartitionWorldUpdater::ReportSectionStatus(UWorld* InWorld) const
+	{
+		FAssetRegistryModule& AssetRegistryModule = FModuleManager::LoadModuleChecked<FAssetRegistryModule>("AssetRegistry");
+		IAssetRegistry* AssetRegistry = &AssetRegistryModule.Get();
+		check(AssetRegistry);
+
+		TMap<FGuid, AMeshPartition*> LoadedMeshPartitions;
+		TMap<FGuid, ACompiledSection*> LoadedCompiledSections;
+		if (InWorld)
+		{
+			for (AMeshPartition* MeshPartition : TActorRange<AMeshPartition>(InWorld))
+			{
+				if (MeshPartition)
+				{
+					const FGuid& MeshPartitionGUID = MeshPartition->GetActorGuid();
+					LoadedMeshPartitions.Add(MeshPartitionGUID, MeshPartition);
+				}
+			}
+			for (ACompiledSection* CompiledSection : TActorRange<ACompiledSection>(InWorld))
+			{
+				if (CompiledSection)
+				{
+					const FGuid& CompiledSectionGUID = CompiledSection->GetActorGuid();
+					LoadedCompiledSections.Add(CompiledSectionGUID, CompiledSection);
+				}
+			}
+		}
+
+		for (const TPair<FGuid, FMeshPartitionUpdater>& Pair : MeshPartitionUpdaters)
+		{
+			const FMeshPartitionUpdater& MeshPartitionUpdater = Pair.Value;
+			TSet<int32> CompiledSectionIndicesReported;
+
+			FGuid MegaMeshGuid = MeshPartitionUpdater.AllCachedDescriptors->GetMegaMeshGUID();
+			AMeshPartition** LoadedActor = LoadedMeshPartitions.Find(MegaMeshGuid);
+
+			FString LoadedActorPath;
+			FString LoadedActorLabel;
+			if (LoadedActor)
+			{
+				LoadedActorPath = (*LoadedActor)->GetPathName();
+				LoadedActorLabel = (*LoadedActor)->GetActorNameOrLabel();
+			}
+
+			FStatusSummary Summary;
+			MeshPartitionUpdater.GetStatusSummary(Summary);
+
+			auto ReportCompiledSection = [&MeshPartitionUpdater, AssetRegistry, &LoadedCompiledSections](int32 CompiledSectionIndex, int32 SectionBaseCount, int32 TotalBaseCount, FName BuildVariantName)
+				{
+					FString StatusString;
+					FString LoadedString;
+					FString DetailsString;
+					bool bReuse = false;
+
+					if (CompiledSectionIndex >= 0)
+					{
+						const FCompiledSectionStatus* Status = MeshPartitionUpdater.GetCompiledSectionStatus(CompiledSectionIndex);
+						check(Status);
+
+						StatusString = Status->ToString();
+
+						const FCompiledSectionDescriptor& CompiledSectionDesc = MeshPartitionUpdater.CompiledSections[CompiledSectionIndex];
+						DetailsString = CompiledSectionDesc.ActorPath.ToString();
+
+						LoadedString = "      ";
+						if (const FCompiledSectionDescriptor* Descriptor = MeshPartitionUpdater.GetCompiledSectionDescriptor(CompiledSectionIndex))
+						{
+							if (ACompiledSection** CompiledSectionPtr = LoadedCompiledSections.Find(Descriptor->ActorDescGuid))
+							{
+								// add loaded details
+								if (ACompiledSection* CompiledSection = *CompiledSectionPtr)
+								{
+									LoadedString = "Loaded";
+
+									if (UTexture2DArray* ChannelTexture = Cast<UTexture2DArray>(CompiledSection->GetChannelTexture()))
+									{
+										DetailsString.Appendf(TEXT(" tex(%dx%dx%d %d/%d mips)"),
+											ChannelTexture->GetSizeX(),
+											ChannelTexture->GetSizeY(),
+											ChannelTexture->GetArraySize(),
+											ChannelTexture->GetNumResidentMips(),
+											//ChannelTexture->GetNumResidentMips(),
+											ChannelTexture->GetNumMips()
+										);
+									}
+
+									TArray<TObjectPtr<UStaticMesh>> StaticMeshes = CompiledSection->GetStaticMeshes();
+									for (TObjectPtr<UStaticMesh> StaticMesh : StaticMeshes)
+									{
+										if (StaticMesh)
+										{
+											int32 LODCount = StaticMesh->GetNumLODs();
+											DetailsString.Appendf(TEXT(" mesh(%d tris %d verts %d lods)"),
+												StaticMesh->GetNumTriangles(0),
+												StaticMesh->GetNumVertices(0),
+												LODCount);
+										}
+									}
+								}
+							}
+						}
+
+						if (!Status->bReuse)
+						{
+							DetailsString = Status->GetDetailsString(CompiledSectionIndex, MeshPartitionUpdater, *AssetRegistry);
+						}
+					}
+					else
+					{
+						StatusString = TEXT("Missing");
+					}
+
+					UE_LOG(LogMegaMeshEditor, Warning, TEXT("  [%s] [%d/%d bases] %s (%s) %s"),
+						*BuildVariantName.ToString(), SectionBaseCount, TotalBaseCount, *LoadedString, *StatusString, *DetailsString);
+				};
+
+			UE_LOG(LogMegaMeshEditor, Warning, TEXT("%s %s (GUID: %s)"), *LoadedActorLabel, *LoadedActorPath, *MegaMeshGuid.ToString());
+
+			// iterate variants
+			for (const TPair<FName, FBuildVariantUpdater>& Pair2 : MeshPartitionUpdater.BuildVariantUpdaters)
+			{
+				const FBuildVariantUpdater& VariantUpdater = Pair2.Value;
+				check(Pair2.Key == VariantUpdater.GetBuildVariantName());
+				FName BuildVariantName = VariantUpdater.GetBuildVariantName();
+
+				UE_LOG(LogMegaMeshEditor, Warning, TEXT(" --- Variant: %s"), *BuildVariantName.ToString());
+
+				int32 TotalBaseCount = MeshPartitionUpdater.AllCachedDescriptors->GetAllBaseModifiers().Num();
+
+				// first report on the groups that reuse a compiled section
+				TConstArrayView<MeshPartition::FModifierGroup> ModifierGroups = VariantUpdater.GetTargetModifierGroups();
+				for (int32 GroupIndex = 0; GroupIndex < ModifierGroups.Num(); GroupIndex++)
+				{
+					const MeshPartition::FModifierGroup& Group = ModifierGroups[GroupIndex];
+					int32 CompiledSectionIndex = VariantUpdater.GetReusableCompiledSectionForGroup(GroupIndex, MeshPartitionUpdater);
+					ReportCompiledSection(CompiledSectionIndex, Group.BaseDescs().Num(), TotalBaseCount, BuildVariantName);
+					if (CompiledSectionIndex >= 0)
+					{
+						CompiledSectionIndicesReported.Add(CompiledSectionIndex);
+					}
+				}
+			}
+
+			UE_LOG(LogMegaMeshEditor, Warning, TEXT(" === REMOVE:"));
+
+			// Now report on all existing compiled sections that were not yet reported - these are not mapped to any target build variant and need to be removed
+			int32 TotalBaseCount = 0; // we can't be sure how many modifiers were originally in this compiled section, without a corresponding grouping
+
+			for (int32 CompiledSectionIndex = 0; CompiledSectionIndex < MeshPartitionUpdater.CompiledSections.Num(); CompiledSectionIndex++)
+			{
+				if (!CompiledSectionIndicesReported.Contains(CompiledSectionIndex))
+				{
+					const FCompiledSectionDescriptor& ExistingDesc = MeshPartitionUpdater.CompiledSections[CompiledSectionIndex];
+					FName BuildVariantName = ExistingDesc.Info.BuildVariantName;
+					ReportCompiledSection(CompiledSectionIndex, ExistingDesc.Info.BaseModifierPaths.Num(), TotalBaseCount, BuildVariantName);
+				}
+			}
+
+			UE_LOG(LogMegaMeshEditor, Warning, TEXT(" ==="));
+
+			UE_LOG(LogMegaMeshEditor, Warning, TEXT(" - Target Groups: %d UpToDate, %d OutOfDate, %d Missing, %d Total"),
+				Summary.GroupsUpToDate(), Summary.GroupsOutOfDate(), Summary.GroupsMissing(), Summary.GroupCount);
+
+			UE_LOG(LogMegaMeshEditor, Warning, TEXT(" - CompiledSections: %d ToRemove, %d Total"),
+				Summary.CompiledSectionsToRemove(), Summary.CompiledSectionCount);
+
+			UE_LOG(LogMegaMeshEditor, Warning, TEXT(" ==="));
+		}
+	}
+
+	FString FCompiledSectionStatus::ToString() const
+	{
+		FString StatusString;
+		check(bChecked);
+
+		if (bReuse)
+		{
+			StatusString = TEXT("UpToDate");
+		}
+		if (bIsDuplicate)
+		{
+			StatusString += TEXT("Duplicate");
+		}
+		if (bCannotFindFile)
+		{
+			StatusString += TEXT("CannotFindFile");
+		}
+		if (bIsPlaceholder)
+		{
+			StatusString += TEXT("Placeholder");
+		}
+		if (bNonTargetVariant)
+		{
+			StatusString += TEXT("NonTargetVariant");
+		}
+		if (bMismatchedModifiers)
+		{
+			StatusString += TEXT("MismatchedModifiers");
+		}
+		if (!bPackageHashMatches)
+		{
+			StatusString += TEXT("PackageHashFail");
+		}
+		if (!bClassHashMatches)
+		{
+			StatusString += TEXT("ClassHashFail");
+		}
+		if (StatusString.IsEmpty())
+		{
+			StatusString = TEXT("<UNKNOWN>");
+		}
+		return StatusString;
+	}
+
+	FString FCompiledSectionStatus::GetDetailsString(int32 CompiledSectionIndex, const FMeshPartitionUpdater& MeshPartitionUpdater, IAssetRegistry& AssetRegistry) const
+	{
+		FString Details;
+		TConstArrayView<FCompiledSectionDescriptor> CompiledSections = MeshPartitionUpdater.AllCachedDescriptors->GetCompiledSections();
+		if (CompiledSections.IsValidIndex(CompiledSectionIndex))
+		{
+			const FCompiledSectionDescriptor& Section = CompiledSections[CompiledSectionIndex];
+
+			if (!bPackageHashMatches)
+			{
+				Details += TEXT(" PackageHashChanged(");
+
+				int32 PackageCount = Section.Info.PackageDependencies.Num();
+				if (Section.Info.PackageChecksums.Num() == PackageCount)
+				{
+					TArray<uint32> CurrentChecksums;
+					FGuid PackageHash = UE::MeshPartition::EditorUtils::ComputePackageHash(AssetRegistry, Section.Info.PackageDependencies, &CurrentChecksums);
+					check(CurrentChecksums.Num() == PackageCount);
+
+					bool First = true;
+					for (int32 PackageIndex = 0; PackageIndex < PackageCount; PackageIndex++)
+					{
+						if (CurrentChecksums[PackageIndex] != Section.Info.PackageChecksums[PackageIndex])
+						{
+							if (!First)
+							{
+								Details += TEXT(", ");
+							}
+							First = false;
+							Details += Section.Info.PackageDependencies[PackageIndex].ToString();
+						}
+					}
+				}
+				else
+				{
+					Details += TEXT("No Checksums Recorded");
+				}
+				Details += TEXT(")");
+			}
+
+			if (!bClassHashMatches)
+			{
+				Details += TEXT(" ClassHashChanged(");
+				int32 ClassCount = Section.Info.ClassDependencies.Num();
+				if (Section.Info.ClassChecksums.Num() == ClassCount)
+				{
+					TArray<const UClass*> Classes;
+					Algo::Transform(Section.Info.ClassDependencies, Classes, [](FName ClassPathName)
+						{
+							FSoftClassPath ClassPath(ClassPathName.ToString());
+							return ClassPath.ResolveClass();
+						});
+
+					TArray<uint32> CurrentChecksums;
+					FGuid ClassHash = UE::MeshPartition::EditorUtils::ComputeClassHash(Classes, &CurrentChecksums);
+					check(CurrentChecksums.Num() == ClassCount);
+
+					bool First = true;
+					for (int32 ClassIndex = 0; ClassIndex < ClassCount; ClassIndex++)
+					{
+						if (CurrentChecksums[ClassIndex] != Section.Info.ClassChecksums[ClassIndex])
+						{
+							if (!First)
+							{
+								Details += TEXT(", ");
+							}
+							First = false;
+							Details += Section.Info.ClassDependencies[ClassIndex].ToString();
+						}
+					}
+				}
+				else
+				{
+					Details += TEXT("No Checksums Recorded");
+				}
+				Details += TEXT(")");
+			}
+		}
+		else
+		{
+			Details += TEXT("<NONE>");
+		}
+
+		return Details;
+	}
 } // namespace UE::MeshPartition
 #undef LOCTEXT_NAMESPACE
\ No newline at end of file
diff --git a/Engine/Plugins/Experimental/MeshPartition/Source/MeshPartitionEditor/Public/MeshPartitionDescriptorCache.h b/Engine/Plugins/Experimental/MeshPartition/Source/MeshPartitionEditor/Public/MeshPartitionDescriptorCache.h
--- a/Engine/Plugins/Experimental/MeshPartition/Source/MeshPartitionEditor/Public/MeshPartitionDescriptorCache.h
+++ b/Engine/Plugins/Experimental/MeshPartition/Source/MeshPartitionEditor/Public/MeshPartitionDescriptorCache.h
@@ -70,6 +70,11 @@ public:
 		return AllModifiersGroup.AllModifierDescs();
 	}
 
+	TConstArrayView<FModifierDesc> GetAllBaseModifiers() const
+	{
+		return AllModifiersGroup.BaseDescs();
+	}
+
 	TConstArrayView<FCompiledSectionDescriptor> GetCompiledSections() const
 	{
 		return CompiledSections;
diff --git a/Engine/Plugins/Experimental/MeshPartition/Source/MeshPartitionEditor/Public/MeshPartitionEditorWorldSubsystem.h b/Engine/Plugins/Experimental/MeshPartition/Source/MeshPartitionEditor/Public/MeshPartitionEditorWorldSubsystem.h
--- a/Engine/Plugins/Experimental/MeshPartition/Source/MeshPartitionEditor/Public/MeshPartitionEditorWorldSubsystem.h
+++ b/Engine/Plugins/Experimental/MeshPartition/Source/MeshPartitionEditor/Public/MeshPartitionEditorWorldSubsystem.h
@@ -19,5 +19,7 @@ public:
 	virtual TStatId GetStatId() const override;
 
 	virtual void Tick(float InDeltaTime) override;
+
+	static void ReportSectionStatus();
 };
 } // namespace UE::MeshPartition
diff --git a/Engine/Plugins/Experimental/MeshPartition/Source/MeshPartitionEditor/Public/MeshPartitionWorldUpdater.h b/Engine/Plugins/Experimental/MeshPartition/Source/MeshPartitionEditor/Public/MeshPartitionWorldUpdater.h
--- a/Engine/Plugins/Experimental/MeshPartition/Source/MeshPartitionEditor/Public/MeshPartitionWorldUpdater.h
+++ b/Engine/Plugins/Experimental/MeshPartition/Source/MeshPartitionEditor/Public/MeshPartitionWorldUpdater.h
@@ -9,6 +9,7 @@
 namespace UE::MeshPartition
 {
 	struct FMeshPartitionUpdater;
+	struct FStatusSummary;
 
 	// what is the state of the prebuilt compiled sections for a group
 	struct FCompiledSectionStatus
@@ -99,8 +100,11 @@ namespace UE::MeshPartition
 		const FBuildVariantUpdater* GetBuildVariantUpdater(const FName& BuildVariantName) const { return BuildVariantUpdaters.Find(BuildVariantName); }
 		void ForAllBuildVariants(TFunctionRef<bool(const MeshPartition::FBuildVariantUpdater&)> InFunc) const;
 
+		const FCompiledSectionDescriptor* GetCompiledSectionDescriptor(int32 CompiledSectionIndex) const { return CompiledSections.IsValidIndex(CompiledSectionIndex) ? &CompiledSections[CompiledSectionIndex] : nullptr; }
 		const FCompiledSectionStatus* GetCompiledSectionStatus(int32 CompiledSectionIndex) const;
 		const FCompiledSectionStatus* GetCompiledSectionStatus(const ACompiledSection* CompiledSection) const;
+
+		void GetStatusSummary(FStatusSummary& InOutSummary) const;
 	};
 
 	// captures the current state of all mesh partitions in a world, and manages what needs to be done to bring the world up to date
@@ -137,5 +141,58 @@ namespace UE::MeshPartition
 		const FCompiledSectionStatus* GetCompiledSectionStatus(const ACompiledSection* CompiledSection) const;
 
 		void ReportStats() const;
+		void GetStatusSummary(FStatusSummary& InOutSummary) const;
+		void ReportSectionStatus(UWorld* InWorld) const;
+	};
+
+	// summary of the status of compiled sections and groups
+	struct FStatusSummary
+	{
+		// overall counts
+		int32 MeshPartitionCount = 0;		// total mesh partitions (existing actors, does not include deleted actors)
+		int32 TargetBuildVariantCount = 0;	// total number of target build variants (sum over all mesh partitions)
+		int32 GroupCount = 0;				// total target modifier groups (sum over all mesh partitions, all target build variants)
+		int32 CompiledSectionCount = 0;		// total existing compiled sections (including out of date sections)
+
+		struct FCompiledSections
+		{
+			// the counts here are mutually exclusive, even if a compiled section meets multiple criteria, it will only count towards one
+			// the exception is "Missing" which indicates compiled sections that do not exist (and can overlap with CannotFindFile)
+
+			// these compiled sections should be removed
+			int32 Placeholders = 0;				// compiled sections that are PIE placeholder leftovers (not cleaned up from previous runs)
+			int32 NonTargetVariant = 0;			// compiled sections that do not correspond to any target build variant
+			int32 MismatchedModifiers = 0;		// compiled sections whose modifier sets do not correspond to any target group
+			int32 Duplicates = 0;				// compiled sections that are redundant duplicates of the same build variant + modifier group
+
+			// these compiled sections need to be updated
+			int32 PackageHashFails = 0;			// compiled sections whose package hashes are out of date and need to be rebuilt, but do correspond to a current group
+			int32 ClassHashFails = 0;			// compiled sections whose class hashes are out of date and need to be rebuilt, but do correspond to a current group
+
+			// these compiled sections can be reused
+			int32 Reuse = 0;					// compiled sections that are up to date and do not need any changes
+
+			// these compiled sections are missing and need to be created
+			int32 Missing = 0;
+
+			// compiled sections that WERE reuseable, but were recently deleted (they still exist in the actor descriptors, but are not in the asset registry) - we can safely ignore these, they will be added to Missing if needed
+			int32 CannotFindFile = 0;
+
+			// this should always be zero, it's counting compiled sections that have none of the above status flags set, which shouldn't happen
+			int32 Unknown = 0;
+		};
+		FCompiledSections CompiledSections;
+
+		// groups that have a corresponding existing compiled section that is up to date
+		int32 GroupsUpToDate() { return CompiledSections.Reuse; }
+
+		// groups that have a corresponding existing compiled section that is NOT up to date
+		int32 GroupsOutOfDate() { return CompiledSections.PackageHashFails + CompiledSections.ClassHashFails; }
+
+		// groups that have NO corresponding existing compiled section
+		int32 GroupsMissing() { return CompiledSections.Missing; }
+
+		// compiled sections that should be deleted
+		int32 CompiledSectionsToRemove() { return CompiledSections.Placeholders + CompiledSections.NonTargetVariant + CompiledSections.MismatchedModifiers + CompiledSections.Duplicates; }
 	};
 } // namespace UE::MeshPartition
\ No newline at end of file
