naming tweaks for mesh terrain erode brush -- 'height erode' -> 'slope erode'; 'directional bias' -> erode strength (w/ sign flip)

#Synced-CL 51500506

diff --git a/Engine/Plugins/Experimental/MeshPartition/Source/MeshPartitionModelingToolset/Private/MeshPartitionHeightSculptTool.cpp b/Engine/Plugins/Experimental/MeshPartition/Source/MeshPartitionModelingToolset/Private/MeshPartitionHeightSculptTool.cpp
--- a/Engine/Plugins/Experimental/MeshPartition/Source/MeshPartitionModelingToolset/Private/MeshPartitionHeightSculptTool.cpp
+++ b/Engine/Plugins/Experimental/MeshPartition/Source/MeshPartitionModelingToolset/Private/MeshPartitionHeightSculptTool.cpp
@@ -95,9 +95,9 @@ void UHeightSculptTool::RegisterBrushes()
 	RegisterBrushType((int32)EBrushType::HeightZero, LOCTEXT("ZeroHeightBrush", "HeightZero"), 
 		MakeUnique<FLambdaMeshSculptBrushOpFactory>([this]() { return MakeUnique<UE::Geometry::FHeightEraseBrushOp>(HeightSculptProperties); }),
 		NewObject<UMeshSculptBrushOpProps>(this, UE::Geometry::FHeightEraseBrushOp::GetPropertiesClass()));
-	RegisterBrushType((int32)EBrushType::HeightErode, LOCTEXT("HeightErodeBrush", "HeightErode"),
-		MakeUnique<FLambdaMeshSculptBrushOpFactory>([this]() { return MakeUnique<UE::Geometry::FHeightErodeBrushOp>(HeightSculptProperties); }),
-		NewObject<UMeshSculptBrushOpProps>(this, UE::Geometry::FHeightErodeBrushOp::GetPropertiesClass()));
+	RegisterBrushType((int32)EBrushType::SlopeErode, LOCTEXT("SlopeErodeBrush", "SlopeErode"),
+		MakeUnique<FLambdaMeshSculptBrushOpFactory>([this]() { return MakeUnique<UE::Geometry::FSlopeErodeBrushOp>(HeightSculptProperties); }),
+		NewObject<UMeshSculptBrushOpProps>(this, UE::Geometry::FSlopeErodeBrushOp::GetPropertiesClass()));
 	RegisterSecondaryBrushType((int32)EBrushType::HeightSmooth, LOCTEXT("SecondarySmoothHeightBrush", "HeightSmooth"),
 		MakeUnique<FLambdaMeshSculptBrushOpFactory>([this]() { return MakeUnique<UE::Geometry::FHeightSmoothBrushOp>(HeightSculptProperties); }),
 		NewObject<UMeshSculptBrushOpProps>(this, UE::Geometry::FHeightSmoothBrushOp::GetPropertiesClass( false)));
diff --git a/Engine/Plugins/Experimental/MeshPartition/Source/MeshPartitionModelingToolset/Private/Sculpting/HeightBrushes.cpp b/Engine/Plugins/Experimental/MeshPartition/Source/MeshPartitionModelingToolset/Private/Sculpting/HeightBrushes.cpp
--- a/Engine/Plugins/Experimental/MeshPartition/Source/MeshPartitionModelingToolset/Private/Sculpting/HeightBrushes.cpp
+++ b/Engine/Plugins/Experimental/MeshPartition/Source/MeshPartitionModelingToolset/Private/Sculpting/HeightBrushes.cpp
@@ -19,10 +19,10 @@ namespace HeightBrushesLocals
 		return FalloffWrapper.Evaluate(Stamp, PositionForFalloff);
 	}
 
-	void ApplyErodeStamp(
+	void ApplySlopeErodeStamp(
 		int Iters,
 		double SlopeThresholdDeg,
-		double DirectionBias,
+		double ErodeStrength,
 		double SmoothingBlend,
 		TFunctionRef<FVector3d(const FVector3d& OrigPos)> GetMoveDirection,
 		const FMeshSculptFallofFunc& FalloffWrapper,
@@ -46,8 +46,8 @@ namespace HeightBrushesLocals
 		double TargetSlopeRad = FMath::DegreesToRadians(SlopeThresholdDeg);
 		double TargetSlopeTan = FMath::Tan(TargetSlopeRad);
 
-		double UpBias = FMath::Clamp(DirectionBias, -1., 1.) * .5 + .5;
-		double DownBias = 1. - UpBias;
+		double DownBias = FMath::Clamp(ErodeStrength, -1., 1.) * .5 + .5;
+		double UpBias = 1. - DownBias;
 
 		for (int32 Iter = 0; Iter < Iters; ++Iter)
 		{
@@ -439,7 +439,7 @@ void UE::Geometry::FHeightEraseBrushOp::ApplyStamp(const FDynamicMesh3* Mesh,
 		Mesh, Stamp, Vertices, NewPositionsOut);
 }
 
-void UE::Geometry::FHeightErodeBrushOp::ApplyStamp(const FDynamicMesh3* Mesh,
+void UE::Geometry::FSlopeErodeBrushOp::ApplyStamp(const FDynamicMesh3* Mesh,
 	const FSculptBrushStamp& Stamp, const TArray<int32>& Vertices, TArray<FVector3d>& NewPositionsOut)
 {
 	using namespace HeightBrushesLocals;
@@ -449,11 +449,11 @@ void UE::Geometry::FHeightErodeBrushOp::ApplyStamp(const FDynamicMesh3* Mesh,
 		const FVector3d& SphereCenter = SphereFrame.Origin;
 		FVector3d FallbackMoveVector = SphereFrame.Z();
 
-		ApplyErodeStamp(
-			GetPropertySetAs<UMeshHeightSculptErodeBrushOpProps>()->Iterations,
-			GetPropertySetAs<UMeshHeightSculptErodeBrushOpProps>()->SlopeThreshold,
-			GetPropertySetAs<UMeshHeightSculptErodeBrushOpProps>()->DirectionBias,
-			GetPropertySetAs<UMeshHeightSculptErodeBrushOpProps>()->SmoothingBlend,
+		ApplySlopeErodeStamp(
+			GetPropertySetAs<UMeshHeightSculptSlopeErodeBrushOpProps>()->Iterations,
+			GetPropertySetAs<UMeshHeightSculptSlopeErodeBrushOpProps>()->SlopeThreshold,
+			GetPropertySetAs<UMeshHeightSculptSlopeErodeBrushOpProps>()->ErodeStrength,
+			GetPropertySetAs<UMeshHeightSculptSlopeErodeBrushOpProps>()->SmoothingBlend,
 			[&SphereCenter, &FallbackMoveVector](const FVector3d& OrigPos)
 			{
 				return GetSphereMoveVector(SphereCenter, OrigPos, FallbackMoveVector);
@@ -466,11 +466,11 @@ void UE::Geometry::FHeightErodeBrushOp::ApplyStamp(const FDynamicMesh3* Mesh,
 
 	FVector3d MoveDirection = Stamp.LocalFrame.Z();
 
-	ApplyErodeStamp(
-		GetPropertySetAs<UMeshHeightSculptErodeBrushOpProps>()->Iterations,
-		GetPropertySetAs<UMeshHeightSculptErodeBrushOpProps>()->SlopeThreshold,
-		GetPropertySetAs<UMeshHeightSculptErodeBrushOpProps>()->DirectionBias,
-		GetPropertySetAs<UMeshHeightSculptErodeBrushOpProps>()->SmoothingBlend,
+	ApplySlopeErodeStamp(
+		GetPropertySetAs<UMeshHeightSculptSlopeErodeBrushOpProps>()->Iterations,
+		GetPropertySetAs<UMeshHeightSculptSlopeErodeBrushOpProps>()->SlopeThreshold,
+		GetPropertySetAs<UMeshHeightSculptSlopeErodeBrushOpProps>()->ErodeStrength,
+		GetPropertySetAs<UMeshHeightSculptSlopeErodeBrushOpProps>()->SmoothingBlend,
 		[&MoveDirection](const FVector3d& OrigPos) { return MoveDirection; },
 		GetFalloff(),
 		Mesh, Stamp, Vertices, NewPositionsOut);
diff --git a/Engine/Plugins/Experimental/MeshPartition/Source/MeshPartitionModelingToolset/Private/Sculpting/HeightBrushes.h b/Engine/Plugins/Experimental/MeshPartition/Source/MeshPartitionModelingToolset/Private/Sculpting/HeightBrushes.h
--- a/Engine/Plugins/Experimental/MeshPartition/Source/MeshPartitionModelingToolset/Private/Sculpting/HeightBrushes.h
+++ b/Engine/Plugins/Experimental/MeshPartition/Source/MeshPartitionModelingToolset/Private/Sculpting/HeightBrushes.h
@@ -64,7 +64,7 @@ public:
 
 // Erosion (slope) brush properties
 UCLASS()
-class UMeshHeightSculptErodeBrushOpProps : public UMeshSculptBrushOpProps
+class UMeshHeightSculptSlopeErodeBrushOpProps : public UMeshSculptBrushOpProps
 {
 	GENERATED_BODY()
 
@@ -86,13 +86,13 @@ public:
 	UPROPERTY(EditAnywhere, Category = HeightBrush, meta=(ClampMin = 1, UIMax = 50))
 	int32 Iterations = 20;
 
-	// Directional bias in applying sloping -- negative values will flatten toward the lower surface, while positive will raise to the higher
+	// Downward strength of erode. Positive erode strength creates mountain-like ridges; negative erode strength will pull up hills.
 	UPROPERTY(EditAnywhere, Category = HeightBrush, meta = (ClampMin = -1, ClampMax = 1))
-	double DirectionBias = -.5;
+	double ErodeStrength = .5;
 
 	// Amount of smoothing to blend into the brush; a value of 1 will apply only smoothing and no erosion/sloping
 	UPROPERTY(EditAnywhere, Category = HeightBrush, meta = (ClampMin = 0, ClampMax = 1))
-	double SmoothingBlend = .1;
+	double SmoothingBlend = 0.;
 
 
 	virtual float GetStrength() override { return Strength; }
@@ -250,10 +250,10 @@ private:
 	FFrame3d ReferencePlane;
 };
 
-class FHeightErodeBrushOp : public FHeightBrushOpBase
+class FSlopeErodeBrushOp : public FHeightBrushOpBase
 {
 public:
-	FHeightErodeBrushOp(UE::MeshPartition::UHeightSculptToolProperties* HeightSculptPropertiesIn)
+	FSlopeErodeBrushOp(UE::MeshPartition::UHeightSculptToolProperties* HeightSculptPropertiesIn)
 		: FHeightBrushOpBase(HeightSculptPropertiesIn)
 	{
 	}
@@ -261,7 +261,7 @@ public:
 	virtual void ApplyStamp(const FDynamicMesh3* SrcMesh, const FSculptBrushStamp& Stamp,
 		const TArray<int32>& Vertices, TArray<FVector3d>& NewPositionsOut) override;
 
-	static UClass* GetPropertiesClass() { return UMeshHeightSculptErodeBrushOpProps::StaticClass(); }
+	static UClass* GetPropertiesClass() { return UMeshHeightSculptSlopeErodeBrushOpProps::StaticClass(); }
 };
 
 }//end UE::Geometry
diff --git a/Engine/Plugins/Experimental/MeshPartition/Source/MeshPartitionModelingToolset/Public/MeshPartitionHeightSculptTool.h b/Engine/Plugins/Experimental/MeshPartition/Source/MeshPartitionModelingToolset/Public/MeshPartitionHeightSculptTool.h
--- a/Engine/Plugins/Experimental/MeshPartition/Source/MeshPartitionModelingToolset/Public/MeshPartitionHeightSculptTool.h
+++ b/Engine/Plugins/Experimental/MeshPartition/Source/MeshPartitionModelingToolset/Public/MeshPartitionHeightSculptTool.h
@@ -68,7 +68,7 @@ public:
 		HeightSmooth,
 		HeightFlatten,
 		HeightZero,
-		HeightErode
+		SlopeErode
 	};
 
 protected:
diff --git a/Engine/Plugins/Experimental/MeshTerrainMode/Source/MeshTerrainMode/Private/MeshTerrainDetailsSection.cpp b/Engine/Plugins/Experimental/MeshTerrainMode/Source/MeshTerrainMode/Private/MeshTerrainDetailsSection.cpp
--- a/Engine/Plugins/Experimental/MeshTerrainMode/Source/MeshTerrainMode/Private/MeshTerrainDetailsSection.cpp
+++ b/Engine/Plugins/Experimental/MeshTerrainMode/Source/MeshTerrainMode/Private/MeshTerrainDetailsSection.cpp
@@ -110,9 +110,9 @@ void FMeshTerrainDetailsSections::RegisterSculptSubmodeTabs()
 		HSculptToolHeightEraseBrushTab.Add("HeightBrush", "Brush");
 		AllSectionMappings.Add(FName("MeshHeightSculptEraseBrushOpProps"), HSculptToolHeightEraseBrushTab);
 
-		TMap<FName, FName> HSculptToolHeightErodeBrushTab;
-		HSculptToolHeightErodeBrushTab.Add("HeightBrush", "Brush");
-		AllSectionMappings.Add(FName("MeshHeightSculptErodeBrushOpProps"), HSculptToolHeightErodeBrushTab);
+		TMap<FName, FName> HSculptToolSlopeErodeBrushTab;
+		HSculptToolSlopeErodeBrushTab.Add("HeightBrush", "Brush");
+		AllSectionMappings.Add(FName("MeshHeightSculptSlopeErodeBrushOpProps"), HSculptToolSlopeErodeBrushTab);
 		
 		TMap<FName, FName> HSculptToolHeightBrushTab;
 		HSculptToolHeightBrushTab.Add("HeightBrush", "Brush");
diff --git a/Engine/Plugins/Experimental/MeshTerrainMode/Source/MeshTerrainMode/Private/MeshTerrainMode.cpp b/Engine/Plugins/Experimental/MeshTerrainMode/Source/MeshTerrainMode/Private/MeshTerrainMode.cpp
--- a/Engine/Plugins/Experimental/MeshTerrainMode/Source/MeshTerrainMode/Private/MeshTerrainMode.cpp
+++ b/Engine/Plugins/Experimental/MeshTerrainMode/Source/MeshTerrainMode/Private/MeshTerrainMode.cpp
@@ -816,7 +816,7 @@ void UMeshTerrainMode::Enter()
 	RegisterHeightSculptBrushTool(ToolManagerCommands.BeginHeightSmoothBrushTool, TEXT("BeginHeightSmoothBrushTool"), (int32)UE::MeshPartition::UHeightSculptTool::EBrushType::HeightSmooth);
 	RegisterHeightSculptBrushTool(ToolManagerCommands.BeginHeightFlattenBrushTool, TEXT("BeginHeightFlattenBrushTool"), (int32)UE::MeshPartition::UHeightSculptTool::EBrushType::HeightFlatten);
 	RegisterHeightSculptBrushTool(ToolManagerCommands.BeginZeroHeightBrushTool, TEXT("BeginZeroHeightBrushTool"), (int32)UE::MeshPartition::UHeightSculptTool::EBrushType::HeightZero);
-	RegisterHeightSculptBrushTool(ToolManagerCommands.BeginHeightErodeBrushTool, TEXT("BeginHeightErodeBrushTool"), (int32)UE::MeshPartition::UHeightSculptTool::EBrushType::HeightErode);
+	RegisterHeightSculptBrushTool(ToolManagerCommands.BeginSlopeErodeBrushTool, TEXT("BeginSlopeErodeBrushTool"), (int32)UE::MeshPartition::UHeightSculptTool::EBrushType::SlopeErode);
 
 	RegisterTool(ToolManagerCommands.BeginPolyEditTool, TEXT("BeginPolyEditTool"), NewObject<UEditMeshPolygonsToolBuilder>());
 	UEditMeshPolygonsToolBuilder* TriEditBuilder = NewObject<UEditMeshPolygonsToolBuilder>();
diff --git a/Engine/Plugins/Experimental/MeshTerrainMode/Source/MeshTerrainMode/Private/MeshTerrainModeManagerActions.cpp b/Engine/Plugins/Experimental/MeshTerrainMode/Source/MeshTerrainMode/Private/MeshTerrainModeManagerActions.cpp
--- a/Engine/Plugins/Experimental/MeshTerrainMode/Source/MeshTerrainMode/Private/MeshTerrainModeManagerActions.cpp
+++ b/Engine/Plugins/Experimental/MeshTerrainMode/Source/MeshTerrainMode/Private/MeshTerrainModeManagerActions.cpp
@@ -313,7 +313,7 @@ void FMeshTerrainModeManagerCommands::RegisterCommands()
 	UI_COMMAND(BeginHeightSmoothBrushTool, "Height Smooth", "Height Smooth", EUserInterfaceActionType::ToggleButton, FInputChord());
 	UI_COMMAND(BeginHeightFlattenBrushTool, "Height Flatten", "Height Flatten", EUserInterfaceActionType::ToggleButton, FInputChord());
 	UI_COMMAND(BeginZeroHeightBrushTool, "Zero Height", "Zero Height", EUserInterfaceActionType::ToggleButton, FInputChord());
-	UI_COMMAND(BeginHeightErodeBrushTool, "Height Erode", "Height Erode", EUserInterfaceActionType::ToggleButton, FInputChord());
+	UI_COMMAND(BeginSlopeErodeBrushTool, "Slope Erode", "Slope Erode", EUserInterfaceActionType::ToggleButton, FInputChord());
 
 	UI_COMMAND(BeginCreateMegaMeshRectangleTool, "Create Rectangle", "Create a new rectangular Mega Mesh", EUserInterfaceActionType::ToggleButton, FInputChord());
 	UI_COMMAND(BeginAddModifierTool, "Add Modifier", "Create a new Mega Mesh modifier", EUserInterfaceActionType::ToggleButton, FInputChord());
diff --git a/Engine/Plugins/Experimental/MeshTerrainMode/Source/MeshTerrainMode/Private/MeshTerrainModeManagerActions.h b/Engine/Plugins/Experimental/MeshTerrainMode/Source/MeshTerrainMode/Private/MeshTerrainModeManagerActions.h
--- a/Engine/Plugins/Experimental/MeshTerrainMode/Source/MeshTerrainMode/Private/MeshTerrainModeManagerActions.h
+++ b/Engine/Plugins/Experimental/MeshTerrainMode/Source/MeshTerrainMode/Private/MeshTerrainModeManagerActions.h
@@ -254,7 +254,7 @@ public:
 	TSharedPtr<FUICommandInfo> BeginHeightSmoothBrushTool;
 	TSharedPtr<FUICommandInfo> BeginHeightFlattenBrushTool;
 	TSharedPtr<FUICommandInfo> BeginZeroHeightBrushTool;
-	TSharedPtr<FUICommandInfo> BeginHeightErodeBrushTool;
+	TSharedPtr<FUICommandInfo> BeginSlopeErodeBrushTool;
 
 	TSharedPtr<FUICommandInfo> BeginCreateMegaMeshRectangleTool;
 	TSharedPtr<FUICommandInfo> BeginAddModifierTool;
diff --git a/Engine/Plugins/Experimental/MeshTerrainMode/Source/MeshTerrainMode/Private/MeshTerrainModeStyle.cpp b/Engine/Plugins/Experimental/MeshTerrainMode/Source/MeshTerrainMode/Private/MeshTerrainModeStyle.cpp
--- a/Engine/Plugins/Experimental/MeshTerrainMode/Source/MeshTerrainMode/Private/MeshTerrainModeStyle.cpp
+++ b/Engine/Plugins/Experimental/MeshTerrainMode/Source/MeshTerrainMode/Private/MeshTerrainModeStyle.cpp
@@ -542,8 +542,8 @@ void FMeshTerrainModeStyle::Initialize()
 		StyleSet->Set("MeshTerrainModeManagerCommands.BeginHeightFlattenBrushTool.Small",	new IMAGE_BRUSH_SVG( "Icons/HeightFlattenBrush", Icon20x20, HeightSculptColor));
 		StyleSet->Set("MeshTerrainModeManagerCommands.BeginZeroHeightBrushTool",			new IMAGE_BRUSH_SVG( "Icons/HeightEraseBrush", Icon40x40, HeightSculptColor));
 		StyleSet->Set("MeshTerrainModeManagerCommands.BeginZeroHeightBrushTool.Small",		new IMAGE_BRUSH_SVG( "Icons/HeightEraseBrush", Icon20x20, HeightSculptColor));
-		StyleSet->Set("MeshTerrainModeManagerCommands.BeginHeightErodeBrushTool",			new IMAGE_BRUSH_SVG("Icons/HeightErosionBrush", Icon40x40, HeightSculptColor));
-		StyleSet->Set("MeshTerrainModeManagerCommands.BeginHeightErodeBrushTool.Small",		new IMAGE_BRUSH_SVG("Icons/HeightErosionBrush", Icon20x20, HeightSculptColor));
+		StyleSet->Set("MeshTerrainModeManagerCommands.BeginSlopeErodeBrushTool",			new IMAGE_BRUSH_SVG("Icons/HeightErosionBrush", Icon40x40, HeightSculptColor));
+		StyleSet->Set("MeshTerrainModeManagerCommands.BeginSlopeErodeBrushTool.Small",		new IMAGE_BRUSH_SVG("Icons/HeightErosionBrush", Icon20x20, HeightSculptColor));
 
 		// Pressure Sensitivity toggle button
 		StyleSet->Set("BrushIcons.PressureSensitivity", new IMAGE_BRUSH_SVG("Icons/PressureSensitivity", Icon20x20));
diff --git a/Engine/Plugins/Experimental/MeshTerrainMode/Source/MeshTerrainMode/Private/Submodes/SculptSubmode.cpp b/Engine/Plugins/Experimental/MeshTerrainMode/Source/MeshTerrainMode/Private/Submodes/SculptSubmode.cpp
--- a/Engine/Plugins/Experimental/MeshTerrainMode/Source/MeshTerrainMode/Private/Submodes/SculptSubmode.cpp
+++ b/Engine/Plugins/Experimental/MeshTerrainMode/Source/MeshTerrainMode/Private/Submodes/SculptSubmode.cpp
@@ -39,7 +39,7 @@ FSculptSubmode::FSculptSubmode(const TSharedPtr<FModeToolkit>& InToolkit)
 		Commands.BeginHeightSmoothBrushTool,
 		Commands.BeginHeightFlattenBrushTool,
 		Commands.BeginZeroHeightBrushTool,
-		Commands.BeginHeightErodeBrushTool,
+		Commands.BeginSlopeErodeBrushTool,
 	};
 	ToolPalettes.Emplace(LOCTEXT("SculptPaletteLabel", "Sculpt"), SculptCommands);
 	ToolPalettes.Emplace(LOCTEXT("HeightSculptPaletteLabel", "Height Sculpt"), HeightSculptCommands);
