Allow the modifier group to have a custom resolver to resolve the modifier pointers.
This is part 1 of solving the issues with PIE with level instances.
It does not change existing logic yet, but allows us to install a way to resolve modifiers in ways other than by path.

#Synced-CL 51480380

diff --git a/Engine/Plugins/Experimental/MeshPartition/Source/MeshPartitionEditor/Private/MeshPartitionModifierDescriptors.cpp b/Engine/Plugins/Experimental/MeshPartition/Source/MeshPartitionEditor/Private/MeshPartitionModifierDescriptors.cpp
--- a/Engine/Plugins/Experimental/MeshPartition/Source/MeshPartitionEditor/Private/MeshPartitionModifierDescriptors.cpp
+++ b/Engine/Plugins/Experimental/MeshPartition/Source/MeshPartitionEditor/Private/MeshPartitionModifierDescriptors.cpp
@@ -718,8 +718,19 @@ void FModifierGroup::ResolveModifierPtrs()
 	for (int32 ModifierIndex = 0; ModifierIndex < ModifierDescriptors.Num(); ModifierIndex++)
 	{
 		const MeshPartition::FModifierDesc& ModifierDesc = GetModifierDesc(FModifierIndex(ModifierIndex));
-		UObject* ModifierObject = ModifierDesc.ModifierPath.ResolveObject();
-		MeshPartition::UModifierComponent* ModifierComponent = Cast<MeshPartition::UModifierComponent>(ModifierObject);
+
+		MeshPartition::UModifierComponent* ModifierComponent = nullptr;
+		if (ModifierResolver)
+		{
+			ModifierComponent = ModifierResolver->ResolveModifier(ModifierDesc);
+		}
+		else
+		{
+			// default behavior is to resolve by path
+			UObject* ModifierObject = ModifierDesc.ModifierPath.ResolveObject();
+			ModifierComponent = Cast<MeshPartition::UModifierComponent>(ModifierObject);
+		}
+
 		if (ModifierComponent != nullptr)
 		{
 			ModifierPointers[ModifierIndex] = ModifierComponent;
diff --git a/Engine/Plugins/Experimental/MeshPartition/Source/MeshPartitionEditor/Public/MeshPartitionModifierDescriptors.h b/Engine/Plugins/Experimental/MeshPartition/Source/MeshPartitionEditor/Public/MeshPartitionModifierDescriptors.h
--- a/Engine/Plugins/Experimental/MeshPartition/Source/MeshPartitionEditor/Public/MeshPartitionModifierDescriptors.h
+++ b/Engine/Plugins/Experimental/MeshPartition/Source/MeshPartitionEditor/Public/MeshPartitionModifierDescriptors.h
@@ -241,6 +241,13 @@ struct TIndexRange
 	int32 End = 0;
 };
 
+class IModifierResolver
+{
+public:
+	virtual MeshPartition::UModifierComponent* ResolveModifier(const MeshPartition::FModifierDesc& ModifierDesc) = 0;
+	virtual ~IModifierResolver() {};
+};
+
 struct FModifierGroup
 {
 public:
@@ -270,6 +277,8 @@ public:
 		return ShouldApplyInstanceBefore(InLayerPriorities, GetInstanceInfo(InInstanceIndex), GetInstanceInfo(InOtherInstanceIndex));
 	}
 
+	void SetModifierResolver(TSharedPtr<MeshPartition::IModifierResolver> InModifierResolver) { this->ModifierResolver = InModifierResolver; }
+
 	/** Returns true if the passed modifiers have a build dependency on each other. */
 	static bool HasDependency(const MeshPartition::FModifierDesc& InModifier, const MeshPartition::FModifierDesc& InOtherModifier);
 	bool HasDependency(const FInstanceInfo& InInstance, const FInstanceInfo& InOtherInstance) const;
@@ -469,6 +478,8 @@ private:
 	/** Group is prepared for use on the async threads. It is no longer permitted to use it to reference GT objects. */
 	bool bAsyncGroup = false;
 
+	TSharedPtr<MeshPartition::IModifierResolver> ModifierResolver;
+
 	EState CurrentState = EState::Uninitialized;
 	TOptional<MeshPartition::EBuildType> BuildType;
 };
