Fix use of Sample() in compute shader.

#rnx
#Synced-CL 51579517

diff --git a/Engine/Plugins/Experimental/MeshPartition/Shaders/MeshPartitionBorderFill.usf b/Engine/Plugins/Experimental/MeshPartition/Shaders/MeshPartitionBorderFill.usf
--- a/Engine/Plugins/Experimental/MeshPartition/Shaders/MeshPartitionBorderFill.usf
+++ b/Engine/Plugins/Experimental/MeshPartition/Shaders/MeshPartitionBorderFill.usf
@@ -121,7 +121,7 @@ void PushCS(uint2 DTid : SV_DispatchThreadID)
 
 	float2 coord = (DTid.xy + (0.5)) * TexelSize;
 	
-	float4 Source = SectionMipIn.Sample(MipSampler, coord);
+	float4 Source = SectionMipIn.SampleLevel(MipSampler, coord, 0);
 	
 	//SectionMipOut[DTid] = float4(coord, 0, 1);
 	SectionMipOut[DTid] = Source;
