UE-365527

Mesh Terrain - Spline Modifier - Stop spline modifier processing

#virtualized
#Synced-CL 51499417

diff --git a/Engine/Plugins/Experimental/MeshPartition/Source/MeshPartitionEditor/Private/Test/MeshPartitionSplineModifierTests.cpp b/Engine/Plugins/Experimental/MeshPartition/Source/MeshPartitionEditor/Private/Test/MeshPartitionSplineModifierTests.cpp
--- a/Engine/Plugins/Experimental/MeshPartition/Source/MeshPartitionEditor/Private/Test/MeshPartitionSplineModifierTests.cpp
+++ b/Engine/Plugins/Experimental/MeshPartition/Source/MeshPartitionEditor/Private/Test/MeshPartitionSplineModifierTests.cpp
@@ -474,6 +474,65 @@ TEST_CLASS_WITH_FLAGS(SplineModifierTests, "MeshPartition.Modifier.Splines", EAu
 				});
 	}
 
+	// QMetry: UE - TC - 20621	
+	TEST_METHOD(Verify_Completely_Disabling_Processing_Of_The_Spline_Modifier)
+	{
+		TestCommandBuilder
+			.Until(TEXT("Wait for initial modifier processing to complete"), [this]()
+				{
+					return !MeshPartitionEditorComponent->IsModifierParticipatingInActivePreviewSectionBuild(SplineModifier);
+				})
+			.Then(TEXT("Set disabled flag to true for the modifier"), [this]()
+				{
+					const FScopedTransaction Transaction(FText::FromString(TEXT("Disable Spline Modifier")));
+					SplineModifier->Modify();
+
+					SplineModifier->SetIsDisabledFlag(true);
+					SplineModifier->UpdateSplineData();
+				})
+			.Until(TEXT("Wait for modifier update to process"), [this]()
+				{
+					return !MeshPartitionEditorComponent->IsAnyPreviewSectionBuildActive();
+				})
+			.Then(TEXT("Verify that Spline modifier is disabled"), [this]()
+				{
+					ASSERT_THAT(IsTrue(SplineModifier->IsDisabled(), TEXT("Spline modifier is not disabled")));
+
+					const FString TestName = TEXT("Verify_Completely_Disabling_Processing_Of_The_Spline_Modifier_Disabled");
+					CompareToReferenceAsset(TestName);
+				})
+			.Then(TEXT("Test Undo function on the feature"), [this]()
+				{
+					GEditor->UndoTransaction();
+				})
+			.Until(TEXT("Wait for modifier update to process"), [this]()
+				{
+					return !MeshPartitionEditorComponent->IsAnyPreviewSectionBuildActive();
+				})
+			.Then(TEXT("Validate the effects of the Undo function: Spline Modifier should be enabled"), [this]()
+				{
+					ASSERT_THAT(IsFalse(SplineModifier->IsDisabled(), TEXT("Undo did not enable the modifier")));
+
+					const FString TestName = TEXT("Verify_Completely_Disabling_Processing_Of_The_Spline_Modifier_Enabled");
+					CompareToReferenceAsset(TestName);
+				})
+			.Then(TEXT("Test Redo function on the feature"), [this]()
+				{
+					GEditor->RedoTransaction();
+				})
+			.Until(TEXT("Wait for modifier update to process"), [this]()
+				{
+					return !MeshPartitionEditorComponent->IsAnyPreviewSectionBuildActive();
+				})
+			.Then(TEXT("Validate the effects of the Redo function: Spline Modifier should be disabled"), [this]()
+				{
+					ASSERT_THAT(IsTrue(SplineModifier->IsDisabled(), TEXT("Redo did not disable the modifier")));
+
+					const FString TestName = TEXT("Verify_Completely_Disabling_Processing_Of_The_Spline_Modifier_Disabled");
+					CompareToReferenceAsset(TestName);
+				});
+	}
+
 	AFTER_EACH()
 	{
 		ScopedEditorWorld.Reset();
