{MeshPartition] Improve channel texture rasterization to deal with tiny islands

#Synced-CL 51589877

diff --git a/Engine/Plugins/Experimental/MeshPartition/Shaders/MeshPartitionMakeSectionChannels.usf b/Engine/Plugins/Experimental/MeshPartition/Shaders/MeshPartitionMakeSectionChannels.usf
--- a/Engine/Plugins/Experimental/MeshPartition/Shaders/MeshPartitionMakeSectionChannels.usf
+++ b/Engine/Plugins/Experimental/MeshPartition/Shaders/MeshPartitionMakeSectionChannels.usf
@@ -7,6 +7,7 @@
 Buffer<int>		InMeshIndices;	
 Buffer<float2>	InMeshUVs;
 Buffer<float>	InMeshChannels;
+Buffer<int>		InMeshOutlines;
 
 uint4  InDrawCall; // NumVertices, NumChannelAttribs, FirstVertIndex
 uint   InChannelOffset; // Offset to the first channel element in the Channel buffer
@@ -20,15 +21,17 @@ void DrawUVDomainVS(
 	out float4  OutWeight : TEXCOORD1,
 	out float4 OutPosition : SV_POSITION)
 {
-	uint VID = InMeshIndices[VertexIndex + InDrawCall.z];
-	
+
+	uint VID = InMeshIndices[VertexIndex + InDrawCall.z]; // Default vertex id
+	if (InDrawCall.w == 1) // OR outline sourcing path
+	{
+		uint EdgeIndex = VertexIndex >> 1;
+		uint EdgeVertexIndex = VertexIndex & 0x00000001;
+		VID = InMeshOutlines[EdgeIndex * 3 + EdgeVertexIndex];	
+	}
+
 	SourceUV = InMeshUVs[VID];
-	
-	// Channel Attrib Address is at ChannelID * NumVertices + VID
-	uint ChannelAttribOffset = InChannelOffset + VID;
-	
-	// Fetch the channel weight from the channel buffer
-	OutWeight = InMeshChannels[ChannelAttribOffset];
+	OutWeight = InMeshChannels[InChannelOffset + VID];
 
 	// Transform from UV space to ndc
 	float2 BaseUV = SourceUV;	
diff --git a/Engine/Plugins/Experimental/MeshPartition/Source/MeshPartition/Public/MeshPartitionMeshData.h b/Engine/Plugins/Experimental/MeshPartition/Source/MeshPartition/Public/MeshPartitionMeshData.h
--- a/Engine/Plugins/Experimental/MeshPartition/Source/MeshPartition/Public/MeshPartitionMeshData.h
+++ b/Engine/Plugins/Experimental/MeshPartition/Source/MeshPartition/Public/MeshPartitionMeshData.h
@@ -528,6 +528,7 @@ struct FSectionChannelGenerationMeshData
 	TArray<int32>			Indices;
 	TArray<FVector2f>		Texcoords;
 	TArray<float>			Channels;
+	TArray<FInt32Vector3>	Outlines;
 	double					DomainArea3D;
 	double					DomainAreaUV;
 	double					DomainGutterUV;
diff --git a/Engine/Plugins/Experimental/MeshPartition/Source/MeshPartitionCompute/Public/MeshPartitionChannelRasterizationShaders.h b/Engine/Plugins/Experimental/MeshPartition/Source/MeshPartitionCompute/Public/MeshPartitionChannelRasterizationShaders.h
--- a/Engine/Plugins/Experimental/MeshPartition/Source/MeshPartitionCompute/Public/MeshPartitionChannelRasterizationShaders.h
+++ b/Engine/Plugins/Experimental/MeshPartition/Source/MeshPartitionCompute/Public/MeshPartitionChannelRasterizationShaders.h
@@ -19,7 +19,8 @@ public:
 	BEGIN_SHADER_PARAMETER_STRUCT(FParameters, )
 		SHADER_PARAMETER_RDG_BUFFER_SRV(Buffer<uint32>, InMeshIndices)
 		SHADER_PARAMETER_RDG_BUFFER_SRV(Buffer<float2>, InMeshUVs)
-		SHADER_PARAMETER_RDG_BUFFER_SRV(Buffer<float>, InMeshChannels)
+		SHADER_PARAMETER_RDG_BUFFER_SRV(Buffer<float>,  InMeshChannels)
+		SHADER_PARAMETER_RDG_BUFFER_SRV(Buffer<uint32>, InMeshOutlines)
 		SHADER_PARAMETER(FUintVector4, InDrawCall)
 		SHADER_PARAMETER(uint32, InChannelOffset)
 	END_SHADER_PARAMETER_STRUCT()
diff --git a/Engine/Plugins/Experimental/MeshPartition/Source/MeshPartitionEditor/Private/MeshPartitionChannelCollection.cpp b/Engine/Plugins/Experimental/MeshPartition/Source/MeshPartitionEditor/Private/MeshPartitionChannelCollection.cpp
--- a/Engine/Plugins/Experimental/MeshPartition/Source/MeshPartitionEditor/Private/MeshPartitionChannelCollection.cpp
+++ b/Engine/Plugins/Experimental/MeshPartition/Source/MeshPartitionEditor/Private/MeshPartitionChannelCollection.cpp
@@ -196,7 +196,7 @@ Tasks::TTask<MeshPartition::FSectionChannels> FChannelTextureRenderer::BuildText
 				.Table = MoveTemp(Section.ChannelOffsetTable),
 				.TexcoordMetrics = Section.TexcoordMetrics,
 				.Texture = ChannelTexture,
-				.Mesh = { MoveTemp(Section.UVmesh.Indices), MoveTemp(Section.UVmesh.UVs), MoveTemp(Section.UVmesh.Channels),
+				.Mesh = { MoveTemp(Section.UVmesh.Indices), MoveTemp(Section.UVmesh.UVs), MoveTemp(Section.UVmesh.Channels), MoveTemp(Section.UVmesh.Outlines),
 							Section.DomainMapping.Area3D, Section.DomainMapping.AreaUV, Section.DomainMapping.GetGutterUV() }
 			};
 
@@ -358,6 +358,7 @@ void FChannelTextureRenderer::RenderSectionChannels(UTexture* InOutSectionTextur
 			int32 NumIndices = Section.UVmesh.NumIndices();
 			int32 NumVertices = Section.UVmesh.NumVertices();
 			int32 NumChannels = Section.UVmesh.NumChannels();
+			int32 NumOutlines = Section.UVmesh.Outlines.Num();
 
 			TArray<uint8> ChannelOffsetTable = Section.ChannelOffsetTable;
 
@@ -407,6 +408,20 @@ void FChannelTextureRenderer::RenderSectionChannels(UTexture* InOutSectionTextur
 				RDGSectionMeshChannelsSRV = GraphBuilder.CreateSRV(RDGSectionBufferDesc);
 			}
 
+			// Section Channels
+			FRDGBufferSRVRef RDGSectionMeshOutlinesSRV;
+			{
+				// Section turned into RDG Resource
+				FRDGBufferRef RDGSectionBuffer = GraphBuilder.CreateBuffer(FRDGBufferDesc::CreateBufferDesc(sizeof(int32), NumOutlines * 3), TEXT("SectionChannelOutlines"));
+				GraphBuilder.QueueBufferUpload(RDGSectionBuffer, Section.UVmesh.Outlines.GetData(), sizeof(int32) * NumOutlines * 3);
+
+				FRDGBufferSRVDesc RDGSectionBufferDesc(RDGSectionBuffer, PF_R32_SINT);
+				RDGSectionBufferDesc.NumElements = NumOutlines * 3;
+				RDGSectionBufferDesc.StartOffsetBytes = 0;
+
+				RDGSectionMeshOutlinesSRV = GraphBuilder.CreateSRV(RDGSectionBufferDesc);
+			}
+
 			// Pass Shaders
 			FGlobalShaderMap* ShaderMap = GetGlobalShaderMap(GMaxRHIFeatureLevel);
 			// Shader with Permutation : auto VertexShader = GlobalShaderMap->GetShader<FRasterizeToRectsVS>(PermutationVector);
@@ -438,190 +453,203 @@ void FChannelTextureRenderer::RenderSectionChannels(UTexture* InOutSectionTextur
 
 			for (int32 Page = 0; Page < Section.NumActivePages(); ++Page)
 			{
-				{
-					// Pass Parameters
-					FMeshPartition_DrawUVDomain_Parameters* PassParameters = GraphBuilder.AllocParameters<FMeshPartition_DrawUVDomain_Parameters>();
-					PassParameters->VS.InMeshIndices = RDGSectionMeshIndicesSRV;
-					PassParameters->VS.InMeshUVs = RDGSectionMeshUVsSRV;
-					PassParameters->VS.InMeshChannels = RDGSectionMeshChannelsSRV;
-					PassParameters->VS.InChannelOffset = Page * NumVertices; // Channel begining offset in the Channels buffer
-
-					ClearUnusedGraphResources(VertexShader, &PassParameters->VS);
-					ClearUnusedGraphResources(PixelShader, &PassParameters->PS);
-
-					PassParameters->PS.RenderTargets[0] = FRenderTargetBinding(RDGRenderTarget, ERenderTargetLoadAction::EClear);
-					PassParameters->PS.RenderTargets[1] = FRenderTargetBinding(RDGMaskTexture, ERenderTargetLoadAction::EClear);
-
-					// Cue Pass
-					GraphBuilder.AddPass(
-						RDG_EVENT_NAME("MeshPartition_MakeSectionChannels"),
-						PassParameters,
-						ERDGPassFlags::Raster,
-						[PassParameters, Page, Section, ShaderMap, VertexShader, PixelShader](FRDGAsyncTask, FRHICommandList& RHICmdList)
-						{
-							// Pass PSO
-							FGraphicsPipelineStateInitializer GraphicsPSOInit;
-							RHICmdList.ApplyCachedRenderTargets(GraphicsPSOInit);
+				// Pass Parameters
+				FMeshPartition_DrawUVDomain_Parameters* PassParameters = GraphBuilder.AllocParameters<FMeshPartition_DrawUVDomain_Parameters>();
+				PassParameters->VS.InMeshIndices = RDGSectionMeshIndicesSRV;
+				PassParameters->VS.InMeshUVs = RDGSectionMeshUVsSRV;
+				PassParameters->VS.InMeshChannels = RDGSectionMeshChannelsSRV;
+				PassParameters->VS.InMeshOutlines = RDGSectionMeshOutlinesSRV;
+				PassParameters->VS.InChannelOffset = Page * NumVertices; // Channel begining offset in the Channels buffer
+
+				ClearUnusedGraphResources(VertexShader, &PassParameters->VS);
+				ClearUnusedGraphResources(PixelShader, &PassParameters->PS);
+
+				PassParameters->PS.RenderTargets[0] = FRenderTargetBinding(RDGRenderTarget, ERenderTargetLoadAction::EClear);
+				PassParameters->PS.RenderTargets[1] = FRenderTargetBinding(RDGMaskTexture, ERenderTargetLoadAction::EClear);
+
+				// Cue Pass
+				GraphBuilder.AddPass(
+					RDG_EVENT_NAME("MeshPartition_MakeSectionChannels"),
+					PassParameters,
+					ERDGPassFlags::Raster,
+					[PassParameters, Page, Section, ShaderMap, VertexShader, PixelShader](FRDGAsyncTask, FRHICommandList& RHICmdList)
+					{
+						// Pass PSO
+						FGraphicsPipelineStateInitializer GraphicsPSOInit;
+						RHICmdList.ApplyCachedRenderTargets(GraphicsPSOInit);
 
-							GraphicsPSOInit.BlendState = TStaticBlendState<>::GetRHI();
-							GraphicsPSOInit.RasterizerState = TStaticRasterizerState<>::GetRHI();
-							GraphicsPSOInit.DepthStencilState = TStaticDepthStencilState<false, CF_Always>::GetRHI();
+						FIntPoint ViewportSize(Section.RasterResolution.X, Section.RasterResolution.Y);
+						RHICmdList.SetViewport(0.0f, 0.0f, 0.0f, (float)ViewportSize.X, (float)ViewportSize.Y, 1.0f);
 
-							GraphicsPSOInit.PrimitiveType = PT_TriangleList;
+						GraphicsPSOInit.BlendState = TStaticBlendState<>::GetRHI();
+						GraphicsPSOInit.RasterizerState = TStaticRasterizerState<>::GetRHI();
+						GraphicsPSOInit.DepthStencilState = TStaticDepthStencilState<false, CF_Always>::GetRHI();
 
-							FVertexDeclarationElementList Elements;
-							GraphicsPSOInit.BoundShaderState.VertexDeclarationRHI = PipelineStateCache::GetOrCreateVertexDeclaration(Elements);
-							GraphicsPSOInit.BoundShaderState.VertexShaderRHI = VertexShader.GetVertexShader();
-							GraphicsPSOInit.BoundShaderState.PixelShaderRHI = PixelShader.GetPixelShader();
+						FVertexDeclarationElementList Elements;
+						GraphicsPSOInit.BoundShaderState.VertexDeclarationRHI = PipelineStateCache::GetOrCreateVertexDeclaration(Elements);
+						GraphicsPSOInit.BoundShaderState.VertexShaderRHI = VertexShader.GetVertexShader();
+						GraphicsPSOInit.BoundShaderState.PixelShaderRHI = PixelShader.GetPixelShader();
 
 
-							SetGraphicsPipelineState(RHICmdList, GraphicsPSOInit, 0);
+						GraphicsPSOInit.PrimitiveType = PT_LineList;
+						SetGraphicsPipelineState(RHICmdList, GraphicsPSOInit, 0);
 
-							FIntPoint ViewportSize(Section.RasterResolution.X, Section.RasterResolution.Y);
-							RHICmdList.SetViewport(0.0f, 0.0f, 0.0f, (float)ViewportSize.X, (float)ViewportSize.Y, 1.0f);
+						for (int32 PassId = 0; PassId < Section.Bases_UVBox.Num(); ++PassId)
+						{
+							const FInt32Vector2 BaseMeshRange = Section.UVmesh.BaseIndicesRanges[PassId];
 
-							for (int32 PassId = 0; PassId < Section.Bases_UVBox.Num(); ++PassId)
-							{
-								const FInt32Vector2 BaseMeshRange = Section.UVmesh.BaseIndicesRanges[PassId];
+							PassParameters->VS.InDrawCall = FUintVector4(Section.UVmesh.NumVertices(), Section.UVmesh.NumChannels(), BaseMeshRange.X, 1);
 
-								PassParameters->VS.InDrawCall = FUintVector4(Section.UVmesh.NumVertices(), Section.UVmesh.NumChannels(), BaseMeshRange.X, BaseMeshRange.Y);
-	
-								SetShaderParameters(RHICmdList, VertexShader, VertexShader.GetVertexShader(), PassParameters->VS);
-								SetShaderParameters(RHICmdList, PixelShader, PixelShader.GetPixelShader(), PassParameters->PS);
+							SetShaderParameters(RHICmdList, VertexShader, VertexShader.GetVertexShader(), PassParameters->VS);
+							SetShaderParameters(RHICmdList, PixelShader, PixelShader.GetPixelShader(), PassParameters->PS);
 
-								int32 NumPrimitives = BaseMeshRange.Y / 3;
-								RHICmdList.DrawPrimitive(0, NumPrimitives, 1);
-							}
-						});
+							int32 NumPrimitives = Section.UVmesh.Outlines.Num();
+							RHICmdList.DrawPrimitive(0, NumPrimitives, 1);
+						}
 
-					// Border fill
-					if (CVarMeshPartition_BorderFillMethod->GetInt() == 1)
-					{
-						FMeshPartition_BorderFillCS::FParameters* BorderParams = GraphBuilder.AllocParameters<FMeshPartition_BorderFillCS::FParameters>();
-						BorderParams->Mask = GraphBuilder.CreateSRV(RDGMaskTexture);
-						BorderParams->RWSectionTexture = GraphBuilder.CreateUAV(RDGRenderTarget);
-						BorderParams->Resolution = FUintVector2(RDGRenderTarget->Desc.Extent.X, RDGRenderTarget->Desc.Extent.Y);
+						GraphicsPSOInit.PrimitiveType = PT_TriangleList;
+						SetGraphicsPipelineState(RHICmdList, GraphicsPSOInit, 0);
+
+						for (int32 PassId = 0; PassId < Section.Bases_UVBox.Num(); ++PassId)
+						{
+							const FInt32Vector2 BaseMeshRange = Section.UVmesh.BaseIndicesRanges[PassId];
+
+							PassParameters->VS.InDrawCall = FUintVector4(Section.UVmesh.NumVertices(), Section.UVmesh.NumChannels(), BaseMeshRange.X, 0);
 
-						TShaderRef<FMeshPartition_BorderFillCS> ComputeShader = ShaderMap->GetShader<FMeshPartition_BorderFillCS>();
+							SetShaderParameters(RHICmdList, VertexShader, VertexShader.GetVertexShader(), PassParameters->VS);
+							SetShaderParameters(RHICmdList, PixelShader, PixelShader.GetPixelShader(), PassParameters->PS);
+
+							int32 NumPrimitives = BaseMeshRange.Y / 3;
+							RHICmdList.DrawPrimitive(0, NumPrimitives, 1);
+						}
+					});
+
+				// Border fill
+				if (CVarMeshPartition_BorderFillMethod->GetInt() == 1)
+				{
+					FMeshPartition_BorderFillCS::FParameters* BorderParams = GraphBuilder.AllocParameters<FMeshPartition_BorderFillCS::FParameters>();
+					BorderParams->Mask = GraphBuilder.CreateSRV(RDGMaskTexture);
+					BorderParams->RWSectionTexture = GraphBuilder.CreateUAV(RDGRenderTarget);
+					BorderParams->Resolution = FUintVector2(RDGRenderTarget->Desc.Extent.X, RDGRenderTarget->Desc.Extent.Y);
+
+					TShaderRef<FMeshPartition_BorderFillCS> ComputeShader = ShaderMap->GetShader<FMeshPartition_BorderFillCS>();
 						
-						FComputeShaderUtils::AddPass(
-							GraphBuilder,
-							RDG_EVENT_NAME("MeshPartition_BorderFill"),
-							ERDGPassFlags::Compute,
-							ComputeShader,
-							BorderParams,
-							FIntVector::DivideAndRoundUp(
-								FIntVector(DestTexture->GetSizeX(), DestTexture->GetSizeY(), 1),
-								FIntVector(8, 8, 1)
-							)
-						);
-
-						// Regular border fill and mip generation
-						FGenerateMips::ExecuteCompute(GraphBuilder, ERHIFeatureLevel::SM6, RDGRenderTarget, TStaticSamplerState<SF_Bilinear, AM_Clamp, AM_Clamp, AM_Clamp>::GetRHI());
-					}
-					// Push pull
-					else if (CVarMeshPartition_BorderFillMethod->GetInt() == 2)
+					FComputeShaderUtils::AddPass(
+						GraphBuilder,
+						RDG_EVENT_NAME("MeshPartition_BorderFill"),
+						ERDGPassFlags::Compute,
+						ComputeShader,
+						BorderParams,
+						FIntVector::DivideAndRoundUp(
+							FIntVector(DestTexture->GetSizeX(), DestTexture->GetSizeY(), 1),
+							FIntVector(8, 8, 1)
+						)
+					);
+
+					// Regular border fill and mip generation
+					FGenerateMips::ExecuteCompute(GraphBuilder, ERHIFeatureLevel::SM6, RDGRenderTarget, TStaticSamplerState<SF_Bilinear, AM_Clamp, AM_Clamp, AM_Clamp>::GetRHI());
+				}
+				// Push pull
+				else if (CVarMeshPartition_BorderFillMethod->GetInt() == 2)
+				{
+					const FRDGTextureDesc& TextureDesc = RDGRenderTarget->Desc;						
+					FRDGTextureSRVDesc SRVDesc(RDGRenderTarget);
+					FRDGTextureUAVDesc UAVDesc(RDGRenderTarget);
+					FRDGTextureSRVDesc SRVMaskDesc(RDGMaskTexture);
+					FRDGTextureUAVDesc UAVMaskDesc(RDGMaskTexture);
+					SRVDesc.NumMipLevels = 1;
+					SRVMaskDesc.NumMipLevels = 1;
+
 					{
-						const FRDGTextureDesc& TextureDesc = RDGRenderTarget->Desc;						
-						FRDGTextureSRVDesc SRVDesc(RDGRenderTarget);
-						FRDGTextureUAVDesc UAVDesc(RDGRenderTarget);
-						FRDGTextureSRVDesc SRVMaskDesc(RDGMaskTexture);
-						FRDGTextureUAVDesc UAVMaskDesc(RDGMaskTexture);
-						SRVDesc.NumMipLevels = 1;
-						SRVMaskDesc.NumMipLevels = 1;
+						TShaderRef<FMeshPartition_FillPullCS> PullShader = ShaderMap->GetShader<FMeshPartition_FillPullCS>();
 
+						// Loop through each level of the mips that require creation and add a dispatch pass per level.
+						for (int8 MipLevel = 1; MipLevel < TextureDesc.NumMips; ++MipLevel)
 						{
-							TShaderRef<FMeshPartition_FillPullCS> PullShader = ShaderMap->GetShader<FMeshPartition_FillPullCS>();
-
-							// Loop through each level of the mips that require creation and add a dispatch pass per level.
-							for (int8 MipLevel = 1; MipLevel < TextureDesc.NumMips; ++MipLevel)
-							{
-								const FIntPoint DestTextureSize(
-									FMath::Max(TextureDesc.Extent.X >> MipLevel, 1),
-									FMath::Max(TextureDesc.Extent.Y >> MipLevel, 1));
-
-								SRVDesc.MipLevel = (int8)(MipLevel - 1);
-								SRVMaskDesc.MipLevel = (int8)(MipLevel - 1);
-								UAVDesc.MipLevel = (int8)(MipLevel);
-								UAVMaskDesc.MipLevel = (int8)(MipLevel);
-
-								FMeshPartition_FillPullCS::FParameters* ShaderParams = GraphBuilder.AllocParameters<FMeshPartition_FillPullCS::FParameters>();
-								ShaderParams->TexelSize = FVector2f(1.0f / DestTextureSize.X, 1.0f / DestTextureSize.Y);
-								ShaderParams->SectionMipIn = GraphBuilder.CreateSRV(SRVDesc);
-								ShaderParams->SectionMipOut = GraphBuilder.CreateUAV(UAVDesc);
-								ShaderParams->MaskMipIn = GraphBuilder.CreateSRV(SRVMaskDesc);
-								ShaderParams->MaskMipOut = GraphBuilder.CreateUAV(UAVMaskDesc);
-								ShaderParams->MipSampler = TStaticSamplerState<SF_Bilinear, AM_Clamp, AM_Clamp, AM_Clamp>::GetRHI();
-								ShaderParams->ResolutionPass = FIntVector4(DestTextureSize.X, DestTextureSize.Y, MipLevel, TextureDesc.NumMips);
-
-								FComputeShaderUtils::AddPass(
-									GraphBuilder,
-									RDG_EVENT_NAME("MeshPartition_FillPull DestMipLevel=%d", MipLevel),
-									ERDGPassFlags::Compute,
-									PullShader,
-									ShaderParams,
-									FComputeShaderUtils::GetGroupCount(DestTextureSize, FComputeShaderUtils::kGolden2DGroupSize)
-								);
-							}
-						}
-						{
-							TShaderRef<FMeshPartition_FillPushCS> PushShader = ShaderMap->GetShader<FMeshPartition_FillPushCS>();
-
-							// Loop through each level of the mips that require creation and add a dispatch pass per level.
-							for (int8 MipLevel = TextureDesc.NumMips - 2; MipLevel >= 0; --MipLevel)
-							{
-								const FIntPoint DestTextureSize(
-									FMath::Max(TextureDesc.Extent.X >> MipLevel, 1),
-									FMath::Max(TextureDesc.Extent.Y >> MipLevel, 1));
-
-								SRVDesc.MipLevel = (int8)(MipLevel + 1);
-								UAVDesc.MipLevel = (int8)(MipLevel);
-								SRVMaskDesc.MipLevel = (int8)(MipLevel);
-
-								FMeshPartition_FillPushCS::FParameters* ShaderParams = GraphBuilder.AllocParameters<FMeshPartition_FillPushCS::FParameters>();
-								ShaderParams->TexelSize = FVector2f(1.0f / DestTextureSize.X, 1.0f / DestTextureSize.Y);
-								ShaderParams->SectionMipIn = GraphBuilder.CreateSRV(SRVDesc);
-								ShaderParams->SectionMipOut = GraphBuilder.CreateUAV(UAVDesc);
-								ShaderParams->MaskMipIn = GraphBuilder.CreateSRV(SRVMaskDesc);
-								ShaderParams->MipSampler = TStaticSamplerState<SF_Bilinear, AM_Clamp, AM_Clamp, AM_Clamp>::GetRHI();
-								ShaderParams->ResolutionPass = FIntVector4(DestTextureSize.X, DestTextureSize.Y, MipLevel, TextureDesc.NumMips);
-
-								FComputeShaderUtils::AddPass(
-									GraphBuilder,
-									RDG_EVENT_NAME("MeshPartition_FillPush DestMipLevel=%d", MipLevel),
-									ERDGPassFlags::Compute,
-									PushShader,
-									ShaderParams,
-									FComputeShaderUtils::GetGroupCount(DestTextureSize, FComputeShaderUtils::kGolden2DGroupSize)
-								);
-							}
+							const FIntPoint DestTextureSize(
+								FMath::Max(TextureDesc.Extent.X >> MipLevel, 1),
+								FMath::Max(TextureDesc.Extent.Y >> MipLevel, 1));
+
+							SRVDesc.MipLevel = (int8)(MipLevel - 1);
+							SRVMaskDesc.MipLevel = (int8)(MipLevel - 1);
+							UAVDesc.MipLevel = (int8)(MipLevel);
+							UAVMaskDesc.MipLevel = (int8)(MipLevel);
+
+							FMeshPartition_FillPullCS::FParameters* ShaderParams = GraphBuilder.AllocParameters<FMeshPartition_FillPullCS::FParameters>();
+							ShaderParams->TexelSize = FVector2f(1.0f / DestTextureSize.X, 1.0f / DestTextureSize.Y);
+							ShaderParams->SectionMipIn = GraphBuilder.CreateSRV(SRVDesc);
+							ShaderParams->SectionMipOut = GraphBuilder.CreateUAV(UAVDesc);
+							ShaderParams->MaskMipIn = GraphBuilder.CreateSRV(SRVMaskDesc);
+							ShaderParams->MaskMipOut = GraphBuilder.CreateUAV(UAVMaskDesc);
+							ShaderParams->MipSampler = TStaticSamplerState<SF_Bilinear, AM_Clamp, AM_Clamp, AM_Clamp>::GetRHI();
+							ShaderParams->ResolutionPass = FIntVector4(DestTextureSize.X, DestTextureSize.Y, MipLevel, TextureDesc.NumMips);
+
+							FComputeShaderUtils::AddPass(
+								GraphBuilder,
+								RDG_EVENT_NAME("MeshPartition_FillPull DestMipLevel=%d", MipLevel),
+								ERDGPassFlags::Compute,
+								PullShader,
+								ShaderParams,
+								FComputeShaderUtils::GetGroupCount(DestTextureSize, FComputeShaderUtils::kGolden2DGroupSize)
+							);
 						}
 					}
-					else
 					{
-						// No border fill but just mip generation
-						FGenerateMips::ExecuteCompute(GraphBuilder, ERHIFeatureLevel::SM6, RDGRenderTarget, TStaticSamplerState<SF_Bilinear, AM_Clamp, AM_Clamp, AM_Clamp>::GetRHI());
-					}
+						TShaderRef<FMeshPartition_FillPushCS> PushShader = ShaderMap->GetShader<FMeshPartition_FillPushCS>();
 
-					// Make sure the texture is actually managed by RDG as a call to RegisterExternalTexture does not guarantee this if the texture was already registered and then had UseExternalAccessMode() called on it.
-					GraphBuilder.UseInternalAccessMode(RDGDestTexture);
-
-					// Copy to regular texture
-					FRHICopyTextureInfo CopyInfo = {
-						.Size = FIntVector(DestTexture->GetSizeX(), DestTexture->GetSizeY(), 1),
-						.SourcePosition = FIntVector(0, 0, 0),
-						.DestPosition = FIntVector(0, 0, 0),
-						.SourceSliceIndex = 0,
-						.DestSliceIndex = (uint32)Page,
-						.NumSlices = 1,
-						.SourceMipIndex = 0,
-						.DestMipIndex = 0,
-						.NumMips = RDGRenderTarget->Desc.NumMips
-					};
-
-					AddCopyTexturePass(GraphBuilder, RDGRenderTarget, RDGDestTexture, CopyInfo);
+						// Loop through each level of the mips that require creation and add a dispatch pass per level.
+						for (int8 MipLevel = TextureDesc.NumMips - 2; MipLevel >= 0; --MipLevel)
+						{
+							const FIntPoint DestTextureSize(
+								FMath::Max(TextureDesc.Extent.X >> MipLevel, 1),
+								FMath::Max(TextureDesc.Extent.Y >> MipLevel, 1));
+
+							SRVDesc.MipLevel = (int8)(MipLevel + 1);
+							UAVDesc.MipLevel = (int8)(MipLevel);
+							SRVMaskDesc.MipLevel = (int8)(MipLevel);
+
+							FMeshPartition_FillPushCS::FParameters* ShaderParams = GraphBuilder.AllocParameters<FMeshPartition_FillPushCS::FParameters>();
+							ShaderParams->TexelSize = FVector2f(1.0f / DestTextureSize.X, 1.0f / DestTextureSize.Y);
+							ShaderParams->SectionMipIn = GraphBuilder.CreateSRV(SRVDesc);
+							ShaderParams->SectionMipOut = GraphBuilder.CreateUAV(UAVDesc);
+							ShaderParams->MaskMipIn = GraphBuilder.CreateSRV(SRVMaskDesc);
+							ShaderParams->MipSampler = TStaticSamplerState<SF_Bilinear, AM_Clamp, AM_Clamp, AM_Clamp>::GetRHI();
+							ShaderParams->ResolutionPass = FIntVector4(DestTextureSize.X, DestTextureSize.Y, MipLevel, TextureDesc.NumMips);
+
+							FComputeShaderUtils::AddPass(
+								GraphBuilder,
+								RDG_EVENT_NAME("MeshPartition_FillPush DestMipLevel=%d", MipLevel),
+								ERDGPassFlags::Compute,
+								PushShader,
+								ShaderParams,
+								FComputeShaderUtils::GetGroupCount(DestTextureSize, FComputeShaderUtils::kGolden2DGroupSize)
+							);
+						}
+					}
+				}
+				else
+				{
+					// No border fill but just mip generation
+					FGenerateMips::ExecuteCompute(GraphBuilder, ERHIFeatureLevel::SM6, RDGRenderTarget, TStaticSamplerState<SF_Bilinear, AM_Clamp, AM_Clamp, AM_Clamp>::GetRHI());
 				}
 
+				// Make sure the texture is actually managed by RDG as a call to RegisterExternalTexture does not guarantee this if the texture was already registered and then had UseExternalAccessMode() called on it.
+				GraphBuilder.UseInternalAccessMode(RDGDestTexture);
+
+				// Copy to regular texture
+				FRHICopyTextureInfo CopyInfo = {
+					.Size = FIntVector(DestTexture->GetSizeX(), DestTexture->GetSizeY(), 1),
+					.SourcePosition = FIntVector(0, 0, 0),
+					.DestPosition = FIntVector(0, 0, 0),
+					.SourceSliceIndex = 0,
+					.DestSliceIndex = (uint32)Page,
+					.NumSlices = 1,
+					.SourceMipIndex = 0,
+					.DestMipIndex = 0,
+					.NumMips = RDGRenderTarget->Desc.NumMips
+				};
+
+				AddCopyTexturePass(GraphBuilder, RDGRenderTarget, RDGDestTexture, CopyInfo);
 			}
 			if (RDGDestTexture)
 			{
diff --git a/Engine/Plugins/Experimental/MeshPartition/Source/MeshPartitionEditor/Private/MeshPartitionChannelCollection_GenerateChannelUV.cpp b/Engine/Plugins/Experimental/MeshPartition/Source/MeshPartitionEditor/Private/MeshPartitionChannelCollection_GenerateChannelUV.cpp
--- a/Engine/Plugins/Experimental/MeshPartition/Source/MeshPartitionEditor/Private/MeshPartitionChannelCollection_GenerateChannelUV.cpp
+++ b/Engine/Plugins/Experimental/MeshPartition/Source/MeshPartitionEditor/Private/MeshPartitionChannelCollection_GenerateChannelUV.cpp
@@ -987,6 +987,72 @@ void FChannelTextureRenderer::GenerateSectionMeshUVs(const FMeshData& InSectionM
 			SectionUVs[VertexId] = InSectionMesh.GetVertexUV(VertexId);
 		});
 
+	// Collect the outline edges
+	TArray<FInt32Vector3> SectionOutlineEdges;
+	{
+		int32 NumTriangles = SectionIndices.Num() / 3;
+		TArray<FInt32Vector4> AllEdges; // Container of the 3 edges per triangle of ALL the triangles of the mesh
+		AllEdges.SetNumUninitialized(3 * NumTriangles);
+
+		// Collect ALL the edges
+		ParallelFor(NumTriangles, [&](int TriangleId)
+			{
+				int32 ElementOffset = TriangleId * 3;
+				int32 A = SectionIndices[ElementOffset + 0];
+				int32 B = SectionIndices[ElementOffset + 1];
+				int32 C = SectionIndices[ElementOffset + 2];
+
+				// Use the (triangleid + 1) to identify the triangle and the sign indicate the edge direction 
+				int32 TID = TriangleId + 1;
+
+				AllEdges[ElementOffset + 0] = (A < B ? FIntVector4(A, B, TID, 0) : FIntVector4(B, A, -TID, 0));
+				AllEdges[ElementOffset + 1] = (B < C ? FIntVector4(B, C, TID, 0) : FIntVector4(C, B, -TID, 0));
+				AllEdges[ElementOffset + 2] = (C < A ? FIntVector4(C, A, TID, 0) : FIntVector4(A, C, -TID, 0));
+			});
+
+		// Sort the edges so the same edges are bundled sequentially
+		AllEdges.Sort([](const FInt32Vector4& A, const FInt32Vector4& B)
+			{
+				if (A.X != B.X) // Sort by X first
+				{
+					return A.X < B.X;
+				}
+				else if (A.Y != B.Y) // Sort by Y second if X are the same
+				{
+					return A.Y < B.Y;
+				}
+				else // Sort by abs(Z) if XY are the same these are the same edges belonging to different triangles
+				{
+					return FMath::Abs(A.Z) < FMath::Abs(B.Z);
+				}
+			});
+
+		// Extract the outline edges which are the the edges with no matching pair
+		FInt32Vector4* CurrentEdge = nullptr;
+		for (int32 eid = 0; eid < AllEdges.Num(); ++eid)
+		{
+			FIntVector4* ThisEdge = AllEdges.GetData() + eid;
+
+			if (!CurrentEdge) // No current at this point so this new edge the current one
+			{
+				CurrentEdge = ThisEdge;
+			}
+			else if ((CurrentEdge->X == ThisEdge->X) && (CurrentEdge->Y == ThisEdge->Y))
+			{
+				// This edge is the matching pair of current edge
+				CurrentEdge->W = ThisEdge->Z;
+				CurrentEdge = nullptr;
+			}
+			else
+			{
+				// Current edge is different from this edge, which means it is an isolated outline edge
+				SectionOutlineEdges.Add({ CurrentEdge->X, CurrentEdge->Y, CurrentEdge->Z });
+
+				// And this edge is the new current
+				CurrentEdge = ThisEdge;
+			}
+		}
+	}
 
 	int32 NumTriangles = SectionIndices.Num() / 3;
 	TArray<int32> SectionIslandPerVertex;
@@ -1012,6 +1078,7 @@ void FChannelTextureRenderer::GenerateSectionMeshUVs(const FMeshData& InSectionM
 	InSection.UVmesh.BaseIndicesRanges = { { 0, SectionIndices.Num() } };
 	InSection.UVmesh.Indices = MoveTemp(SectionIndices);
 	InSection.UVmesh.UVs = MoveTemp(SectionUVs);
+	InSection.UVmesh.Outlines = MoveTemp(SectionOutlineEdges);
 
 	InSection.RasterResolution = FUint32Vector2((uint32)InOutDomainMapping.ImageResolution);
 	InSection.TexcoordMetrics = FVector2f(InOutDomainMapping.Size3D);
diff --git a/Engine/Plugins/Experimental/MeshPartition/Source/MeshPartitionEditor/Public/MeshPartitionChannelCollection.h b/Engine/Plugins/Experimental/MeshPartition/Source/MeshPartitionEditor/Public/MeshPartitionChannelCollection.h
--- a/Engine/Plugins/Experimental/MeshPartition/Source/MeshPartitionEditor/Public/MeshPartitionChannelCollection.h
+++ b/Engine/Plugins/Experimental/MeshPartition/Source/MeshPartitionEditor/Public/MeshPartitionChannelCollection.h
@@ -108,6 +108,7 @@ private:
 			TArray<FInt32Vector2>	BaseIndicesRanges;
 			TArray<int32>			Indices;
 			TArray<FVector2f>		UVs;
+			TArray<FInt32Vector3>	Outlines;
 			TArray<float>			Channels;
 			TArray<int32>			ChannelActiveCounts;
 			TArray<uint8>			ActiveChannelIndices;
