Add controls for nanite fallback to the mesh partition static mesh transformer

#Synced-CL 51539202

diff --git a/Engine/Plugins/Experimental/MeshPartition/Source/MeshPartitionEditor/Private/MeshPartitionEditorUtils.cpp b/Engine/Plugins/Experimental/MeshPartition/Source/MeshPartitionEditor/Private/MeshPartitionEditorUtils.cpp
--- a/Engine/Plugins/Experimental/MeshPartition/Source/MeshPartitionEditor/Private/MeshPartitionEditorUtils.cpp
+++ b/Engine/Plugins/Experimental/MeshPartition/Source/MeshPartitionEditor/Private/MeshPartitionEditorUtils.cpp
@@ -352,14 +352,13 @@ namespace UE::MeshPartition::EditorUtils
 			++CurNumValidSections;
 		}
 
+		StaticMesh->GetNaniteSettings().bEnabled = InParams.bUseNanite;
 		if (InParams.bUseNanite)
 		{
-			StaticMesh->GetNaniteSettings().bEnabled = InParams.bUseNanite;
-
-			// setup fallback mesh at 100% (this drives collision on the nanite mesh)
-			StaticMesh->GetNaniteSettings().GenerateFallback = ENaniteGenerateFallback::Enabled;
-			StaticMesh->GetNaniteSettings().FallbackTarget = ENaniteFallbackTarget::PercentTriangles;
-			StaticMesh->GetNaniteSettings().FallbackPercentTriangles = 1.0f;
+			StaticMesh->GetNaniteSettings().GenerateFallback = InParams.NaniteFallbackMode;
+			StaticMesh->GetNaniteSettings().FallbackTarget = InParams.NaniteFallbackTarget;
+			StaticMesh->GetNaniteSettings().FallbackPercentTriangles = InParams.NaniteFallbackPercentTriangles;
+			StaticMesh->GetNaniteSettings().FallbackRelativeError = InParams.NaniteFallbackRelativeError;
 		}
 
 		StaticMesh->bSupportRayTracing = true;
diff --git a/Engine/Plugins/Experimental/MeshPartition/Source/MeshPartitionEditor/Private/MeshPartitionStaticMeshTransformer.cpp b/Engine/Plugins/Experimental/MeshPartition/Source/MeshPartitionEditor/Private/MeshPartitionStaticMeshTransformer.cpp
--- a/Engine/Plugins/Experimental/MeshPartition/Source/MeshPartitionEditor/Private/MeshPartitionStaticMeshTransformer.cpp
+++ b/Engine/Plugins/Experimental/MeshPartition/Source/MeshPartitionEditor/Private/MeshPartitionStaticMeshTransformer.cpp
@@ -78,8 +78,12 @@ void FStaticMeshTransformer::GatherDependencies(MeshPartition::IDependencyInterf
 
 	InDependencies += LODStreaming;
 	InDependencies += NumStreamedLODs;
-	InDependencies += bUseNanite;
 	InDependencies += bCanEverAffectNavigation;
+	InDependencies += bUseNanite;
+	InDependencies += NaniteFallbackMode;
+	InDependencies += NaniteFallbackTarget;
+	InDependencies += NaniteFallbackPercentTriangles;
+	InDependencies += NaniteFallbackRelativeError;
 }
 
 TArray<UStaticMesh*> FStaticMeshTransformer::CreateStaticMeshes(MeshPartition::FTransformerContext& InTransformerContext) const
@@ -280,8 +284,12 @@ void FStaticMeshTransformer::FinalizeStaticMeshes(MeshPartition::FTransformerCon
 			Params.CollisionProfile = CollisionProfile;
 			Params.NumLODs = LODs.Num();
 			Params.bCanEverAffectNavigation = bCanEverAffectNavigation;
-			Params.bUseNanite = bUseNanite;
 			Params.bSetupSections = true;
+			Params.bUseNanite = bUseNanite;
+			Params.NaniteFallbackMode = NaniteFallbackMode;
+			Params.NaniteFallbackTarget = NaniteFallbackTarget;
+			Params.NaniteFallbackPercentTriangles = NaniteFallbackPercentTriangles;
+			Params.NaniteFallbackRelativeError = NaniteFallbackRelativeError;
 
 			if (MeshPartition::APreviewSection* PreviewSection = Cast<MeshPartition::APreviewSection>(TransformerUnit.Section))
 			{
diff --git a/Engine/Plugins/Experimental/MeshPartition/Source/MeshPartitionEditor/Public/MeshPartitionEditorUtils.h b/Engine/Plugins/Experimental/MeshPartition/Source/MeshPartitionEditor/Public/MeshPartitionEditorUtils.h
--- a/Engine/Plugins/Experimental/MeshPartition/Source/MeshPartitionEditor/Public/MeshPartitionEditorUtils.h
+++ b/Engine/Plugins/Experimental/MeshPartition/Source/MeshPartitionEditor/Public/MeshPartitionEditorUtils.h
@@ -3,6 +3,7 @@
 #pragma once
 
 #include "Engine/CollisionProfile.h"
+#include "Engine/EngineTypes.h"
 #include "MeshPartitionEditorModule.h"
 #include "MeshPartitionMeshData.h"
 
@@ -108,6 +109,10 @@ namespace UE::MeshPartition::EditorUtils
 		bool bCanEverAffectNavigation = false;
 		bool bUseNanite = true;
 		bool bSetupSections = true;
+		ENaniteGenerateFallback NaniteFallbackMode = ENaniteGenerateFallback::Enabled;
+		ENaniteFallbackTarget NaniteFallbackTarget = ENaniteFallbackTarget::PercentTriangles;
+		float NaniteFallbackPercentTriangles = 0.2f;
+		float NaniteFallbackRelativeError = 1.0f;
 	};
 
 	/** Finalizes the build of the provided static mesh given the settings contained in the passed struct. */
diff --git a/Engine/Plugins/Experimental/MeshPartition/Source/MeshPartitionEditor/Public/MeshPartitionStaticMeshTransformer.h b/Engine/Plugins/Experimental/MeshPartition/Source/MeshPartitionEditor/Public/MeshPartitionStaticMeshTransformer.h
--- a/Engine/Plugins/Experimental/MeshPartition/Source/MeshPartitionEditor/Public/MeshPartitionStaticMeshTransformer.h
+++ b/Engine/Plugins/Experimental/MeshPartition/Source/MeshPartitionEditor/Public/MeshPartitionStaticMeshTransformer.h
@@ -52,10 +52,22 @@ private:
 	FPerPlatformInt NumStreamedLODs;
 
 	UPROPERTY(EditAnywhere, Category = "Static Mesh")
+	bool bCanEverAffectNavigation = false;
+
+	UPROPERTY(EditAnywhere, Category = "Nanite")
 	bool bUseNanite = true;
 
-	UPROPERTY(EditAnywhere, Category = "Static Mesh")
-	bool bCanEverAffectNavigation = false;
+	UPROPERTY(EditAnywhere, Category = "Nanite", meta = (EditCondition = "bUseNanite", DisplayName = "Generate Fallback Mesh"))
+	ENaniteGenerateFallback NaniteFallbackMode = ENaniteGenerateFallback::Enabled;
+
+	UPROPERTY(EditAnywhere, Category = "Nanite", meta = (EditCondition = "bUseNanite"))
+	ENaniteFallbackTarget NaniteFallbackTarget = ENaniteFallbackTarget::PercentTriangles;
+
+	UPROPERTY(EditAnywhere, Category = "Nanite", meta = (EditCondition = "bUseNanite && NaniteFallbackTarget == ENaniteFallbackTarget::PercentTriangles", ClampMin = "0.0", ClampMax = "1.0"))
+	float NaniteFallbackPercentTriangles = 0.2f;
+
+	UPROPERTY(EditAnywhere, Category = "Nanite", meta = (EditCondition = "bUseNanite && NaniteFallbackTarget == ENaniteFallbackTarget::RelativeError", ClampMin = "0.0"))
+	float NaniteFallbackRelativeError = 1.0f;
 };
 
 } // namespace UE::MeshPartition
