MeshPartitionUI: Enable bi-directional selection between the mesh partition outliner and the level outliner. Enable the TedsViewer widget by default.

#Synced-CL 51607007

diff --git a/Engine/Plugins/Experimental/MeshPartition/Source/MeshPartitionEditorUI/Private/UI/MeshPartitionSettingsWidget.cpp b/Engine/Plugins/Experimental/MeshPartition/Source/MeshPartitionEditorUI/Private/UI/MeshPartitionSettingsWidget.cpp
--- a/Engine/Plugins/Experimental/MeshPartition/Source/MeshPartitionEditorUI/Private/UI/MeshPartitionSettingsWidget.cpp
+++ b/Engine/Plugins/Experimental/MeshPartition/Source/MeshPartitionEditorUI/Private/UI/MeshPartitionSettingsWidget.cpp
@@ -73,7 +73,7 @@
 
 static FAutoConsoleVariable CVarMeshPartitionOutlinerEnableTedsHierarchyWidget(
 	TEXT("MeshPartitionOutliner.EnableTedsHierarchyWidget"),
-	false,
+	true,
 	TEXT("Enables new Teds Hierarchy Widget instead of the Teds Outliner Widget in the Mesh Partition Outliner"),
 	ECVF_Default
 );
@@ -100,6 +100,14 @@ void SMegaMeshSettingsWidget::PrivateRegisterAttributes(FSlateAttributeInitializ
 
 SMegaMeshSettingsWidget::SMegaMeshSettingsWidget()
 {
+	USelection::SelectionChangedEvent.AddRaw(this, &SMegaMeshSettingsWidget::OnLevelSelectionChanged);
+	USelection::SelectObjectEvent.AddRaw(this, &SMegaMeshSettingsWidget::OnLevelSelectionChanged);
+}
+
+SMegaMeshSettingsWidget::~SMegaMeshSettingsWidget()
+{
+	USelection::SelectionChangedEvent.RemoveAll(this);
+	USelection::SelectObjectEvent.RemoveAll(this);
 }
 
 namespace SMegaMeshSettingsWidgetLocals
@@ -595,7 +603,7 @@ TSharedRef<SSceneOutliner> SMegaMeshSettingsWidget::CreateOutlinerWidget()
 	{
 		SelectionChangeColumn->OnSelectionChanged.AddLambda([this]()
 		{
-			OnOutlinerSelectionChanged();
+			OnOutlinerSelectionChanged(ESelectInfo::Direct);
 		});
 	}
 
@@ -687,7 +695,7 @@ TSharedRef<UE::Editor::DataStorage::STedsCompositeHierarchyViewer> SMegaMeshSett
 			.PrimaryColumn(FTypedElementLabelColumn::StaticStruct())
 			.OnSelectionChanged_Lambda([this](RowHandle Row, ESelectInfo::Type SelectInfo)
 				{
-					OnOutlinerSelectionChanged();
+					OnOutlinerSelectionChanged(SelectInfo);
 				})
 			.EmptyRowsMessage(LOCTEXT("MegaMeshSettings_NoRows", "No modifiers found for this Mesh Partition instance.")))
 		.Filters(Filters)
@@ -848,9 +856,13 @@ FText SMegaMeshSettingsWidget::GetActiveDefinitionLabel() const
 	}
 }
 
-void SMegaMeshSettingsWidget::OnOutlinerSelectionChanged()
+void SMegaMeshSettingsWidget::OnOutlinerSelectionChanged(ESelectInfo::Type SelectInfo)
 {
 	UpdateDetailsViewFromOutliner();
+	if (SelectInfo != ESelectInfo::Direct)
+	{
+		UpdateLevelSelectionFromOutliner();
+	}
 }
 
 void SMegaMeshSettingsWidget::UpdateDetailsViewFromOutliner()
@@ -911,6 +923,200 @@ void SMegaMeshSettingsWidget::UpdateDetailsViewFromOutliner()
 	DetailsView->SetObject(nullptr);
 }
 
+void SMegaMeshSettingsWidget::UpdateLevelSelectionFromOutliner()
+{
+	TArray<UObject*> Modifiers;
+
+	if (Outliner.IsValid())
+	{
+		TArray<FSceneOutlinerTreeItemPtr> SelectedItems = Outliner->GetSelectedItems();
+		Modifiers.Reserve(SelectedItems.Num());
+		for (FSceneOutlinerTreeItemPtr& SelectedItem : SelectedItems)
+		{
+			UActorComponent* Modifier = SMegaMeshSettingsWidgetLocals::GetModifierFromOutlinerItem(*SelectedItem);
+			if (Modifier)
+			{
+				Modifiers.Add(Modifier);
+			}
+		}
+
+		if (Modifiers.Num() > 0)
+		{
+			for (UObject* Modifier : Modifiers)
+			{
+				GEditor->SelectComponent(Cast<UActorComponent>(Modifier), true, true);
+			}
+		}
+	}
+
+	if (HierarchyViewer)
+	{
+		using namespace Editor::DataStorage;
+		ICoreProvider* DataStorage = GetDataStorage();
+		if (!DataStorage)
+		{
+			return;
+		}
+
+		HierarchyViewer->ForEachSelectedRow([this, DataStorage, &Modifiers](RowHandle SelectedRow)
+			{
+				const FTypedElementUObjectColumn* RawObjectColumn = DataStorage->GetColumn<FTypedElementUObjectColumn>(SelectedRow);
+				UActorComponent* Modifier = RawObjectColumn ? Cast<UActorComponent>(RawObjectColumn->Object) : nullptr;
+
+				if (Modifier != nullptr)
+				{
+					Modifiers.Add(Modifier);
+				}
+
+			});
+
+
+		if (Modifiers.Num() > 0)
+		{
+			GEditor->SelectNone(true, true);
+			for (UObject* Modifier : Modifiers)
+			{
+				GEditor->SelectComponent(Cast<UActorComponent>(Modifier), true, true);
+			}
+		}
+
+	}
+}
+
+
+void SMegaMeshSettingsWidget::OnLevelSelectionChanged(UObject* InObject)
+{
+	using namespace UE::Editor::DataStorage;
+	ICoreProvider* DataStorage = GetDataStorage();
+
+	if (Outliner.IsValid())
+	{
+		if (DataStorage)
+		{
+			UMegaMeshTedsFactory* Factory = DataStorage->FindFactory<UMegaMeshTedsFactory>();
+			if (!Factory)
+			{
+				return;
+			}
+
+			TArray<RowHandle> ModifierRows;
+			DataStorage->RunQuery(Factory->MegaMeshModifierQueryRW, Queries::CreateDirectQueryCallbackBinding(
+				[&ModifierRows](IDirectQueryContext& Context, RowHandle Row)
+				{
+					ModifierRows.Add(Row);
+				}));
+			for (RowHandle ModifierRow : ModifierRows)
+			{
+				DataStorage->RemoveColumns<FTypedElementSelectionColumn>(ModifierRow);
+			}
+
+			USelection* Selection = Cast<USelection>(InObject);
+
+			if (!Selection)
+			{
+				return;
+			}
+			TArray<USceneComponent*> SelectedComponents;
+			TArray<UModifierComponent*> SelectedModifiers;
+			TArray<AActor*> SelectedActors;
+			Selection->GetSelectedObjects<USceneComponent>(SelectedComponents);
+			Selection->GetSelectedObjects<UModifierComponent>(SelectedModifiers);
+			Selection->GetSelectedObjects<AActor>(SelectedActors);
+
+			if (SelectedComponents.Num() > SelectedModifiers.Num())
+			{
+				SelectedModifiers.RemoveAll([&SelectedComponents](UModifierComponent* Modifier)
+					{
+						return SelectedComponents.Contains(Modifier);
+					});
+			}
+
+			for (UModifierComponent* Modifier : SelectedModifiers)
+			{
+				RowHandle ModifierRow = DataStorage->LookupMappedRow(ICompatibilityProvider::ObjectMappingDomain,
+					FMapKeyView(Modifier));
+
+				if (Modifier->IsSelected() && ModifierRow != InvalidRowHandle)
+				{
+					DataStorage->AddColumn(ModifierRow, FTypedElementSelectionColumn{ .SelectionSet = FName() });
+				}
+			}
+
+			for (AActor* Actor : SelectedActors)
+			{
+				TInlineComponentArray< UModifierComponent*, 10> OwnedModifiers;
+				Actor->GetComponents<UModifierComponent>(OwnedModifiers, false);
+
+				for (UModifierComponent* Modifier : OwnedModifiers)
+				{
+					RowHandle ModifierRow = DataStorage->LookupMappedRow(ICompatibilityProvider::ObjectMappingDomain,
+						FMapKeyView(Modifier));
+					if (ModifierRow != InvalidRowHandle)
+					{
+						DataStorage->AddColumn(ModifierRow, FTypedElementSelectionColumn{ .SelectionSet = FName() });
+					}
+				}
+			}
+		}
+	}
+
+	if (HierarchyViewer && DataStorage)
+	{
+
+		USelection* Selection = Cast<USelection>(InObject);
+
+		if (!Selection)
+		{
+			return;
+		}
+		TArray<USceneComponent*> SelectedComponents;
+		TArray<UModifierComponent*> SelectedModifiers;
+		TArray<AActor*> SelectedActors;
+		Selection->GetSelectedObjects<USceneComponent>(SelectedComponents);
+		Selection->GetSelectedObjects<UModifierComponent>(SelectedModifiers);
+		Selection->GetSelectedObjects<AActor>(SelectedActors);
+
+		if (SelectedComponents.Num() > SelectedModifiers.Num())
+		{
+			SelectedModifiers.RemoveAll([&SelectedComponents](UModifierComponent* Modifier)
+				{
+					return SelectedComponents.Contains(Modifier);
+				});
+		}
+
+		HierarchyViewer->ClearSelection();
+
+		for (UModifierComponent* Modifier : SelectedModifiers)
+		{
+			RowHandle ModifierRow = DataStorage->LookupMappedRow(ICompatibilityProvider::ObjectMappingDomain,
+				FMapKeyView(Modifier));
+
+			if (ModifierRow != InvalidRowHandle)
+			{
+				HierarchyViewer->SetSelection(ModifierRow, true, ESelectInfo::Direct);					
+			}
+		}
+
+		for (AActor* Actor : SelectedActors)
+		{
+			TInlineComponentArray< UModifierComponent*, 10> OwnedModifiers;
+			Actor->GetComponents<UModifierComponent>(OwnedModifiers, false);
+
+			for (UModifierComponent* Modifier : OwnedModifiers)
+			{
+				RowHandle ModifierRow = DataStorage->LookupMappedRow(ICompatibilityProvider::ObjectMappingDomain,
+					FMapKeyView(Modifier));
+				if (ModifierRow != InvalidRowHandle)
+				{
+					HierarchyViewer->SetSelection(ModifierRow, true, ESelectInfo::Direct);
+				}
+			}
+		}
+
+	}
+
+}
+
 
 void SMegaMeshSettingsWidget::OpenDefinitionInContentBrowser()
 {
diff --git a/Engine/Plugins/Experimental/MeshPartition/Source/MeshPartitionEditorUI/Private/UI/MeshPartitionSettingsWidget.h b/Engine/Plugins/Experimental/MeshPartition/Source/MeshPartitionEditorUI/Private/UI/MeshPartitionSettingsWidget.h
--- a/Engine/Plugins/Experimental/MeshPartition/Source/MeshPartitionEditorUI/Private/UI/MeshPartitionSettingsWidget.h
+++ b/Engine/Plugins/Experimental/MeshPartition/Source/MeshPartitionEditorUI/Private/UI/MeshPartitionSettingsWidget.h
@@ -87,6 +87,7 @@ public:
 	SLATE_END_ARGS();
 
 	UE_API SMegaMeshSettingsWidget();
+	UE_API ~SMegaMeshSettingsWidget();
 	UE_API void Construct(const FArguments& InArgs);
 	UE_API void Tick(const FGeometry& AllottedGeometry, const double InCurrentTime, const float InDeltaTime) override;
 	UE_API virtual void AddReferencedObjects(FReferenceCollector& Collector) override;
@@ -127,10 +128,12 @@ private:
 	UE_API FReply OnMegaMeshSelectClicked();
 	UE_API bool IsActiveMegaMesh() const;
 	UE_API void RebuildMegaMeshSelectionList();
-	UE_API void OnOutlinerSelectionChanged();
+	UE_API void OnOutlinerSelectionChanged(ESelectInfo::Type SelectInfo);
 	UE_API void UpdateDetailsViewFromOutliner();
+	UE_API void UpdateLevelSelectionFromOutliner();
 
 	UE_API void OnFilterSettingsModified(UObject*, FProperty*);
+	UE_API void OnLevelSelectionChanged(UObject* InObject);
 
 	UE_API Editor::DataStorage::ICoreProvider* GetDataStorage();
 	UE_API Editor::DataStorage::IUiProvider* GetDataStorageUI();
