UE-361650

Mesh Terrain - Spline Modifier Test - Fix for Extending Spline falloff bounds test: Replacing the use of FProperty with Getter/Setter

#virtualized
#Synced-CL 51499417

diff --git a/Engine/Plugins/Experimental/MeshPartition/Source/MeshPartitionEditor/Private/Modifiers/MeshPartitionSplineModifier.cpp b/Engine/Plugins/Experimental/MeshPartition/Source/MeshPartitionEditor/Private/Modifiers/MeshPartitionSplineModifier.cpp
--- a/Engine/Plugins/Experimental/MeshPartition/Source/MeshPartitionEditor/Private/Modifiers/MeshPartitionSplineModifier.cpp
+++ b/Engine/Plugins/Experimental/MeshPartition/Source/MeshPartitionEditor/Private/Modifiers/MeshPartitionSplineModifier.cpp
@@ -1042,4 +1042,13 @@ void USplineModifier::SetMeshedInteriorNumTriTarget(double InMeshedInteriorNumTr
 	MeshedInteriorNumTriTarget = FMath::Max(1.0, InMeshedInteriorNumTriTarget);
 }
 
+void USplineModifier::SetFalloffBoundsMultiplier(float InFalloffBoundsMultiplier)
+{
+	if (!FMath::IsFinite(InFalloffBoundsMultiplier))
+	{
+		return;
+	}
+
+	FalloffBoundsMultiplier = FMath::Max(1.0, InFalloffBoundsMultiplier);
+}
 } // namespace UE::MeshPartition
\ No newline at end of file
diff --git a/Engine/Plugins/Experimental/MeshPartition/Source/MeshPartitionEditor/Private/Test/MeshPartitionSplineModifierTests.cpp b/Engine/Plugins/Experimental/MeshPartition/Source/MeshPartitionEditor/Private/Test/MeshPartitionSplineModifierTests.cpp
--- a/Engine/Plugins/Experimental/MeshPartition/Source/MeshPartitionEditor/Private/Test/MeshPartitionSplineModifierTests.cpp
+++ b/Engine/Plugins/Experimental/MeshPartition/Source/MeshPartitionEditor/Private/Test/MeshPartitionSplineModifierTests.cpp
@@ -62,7 +62,7 @@ TEST_CLASS_WITH_FLAGS(SplineModifierTests, "MeshPartition.Modifier.Splines", EAu
 	/*
 	 * Specifies the falloff bounds multiplier for the relevant test.
 	 */
-	const float FalloffBoundsMultiplier = 1.0f;
+	const float FalloffBounds = 10.0f;
 
 	/* 
 	 * Specifies the Z axis value to be set for projection bounds test.
@@ -184,7 +184,7 @@ TEST_CLASS_WITH_FLAGS(SplineModifierTests, "MeshPartition.Modifier.Splines", EAu
 	// QMetry: UE - TC - 20617	
 	TEST_METHOD(Verify_Extending_Spline_Falloff_Bounds)
 	{
-		UpdateFalloutBoundsTest(FalloffBoundsMultiplier);
+		ExtendSplineFalloffBoundsTest(FalloffBounds);
 	}
 
 	// QMetry: UE - TC - 20615	
@@ -723,48 +723,59 @@ TEST_CLASS_WITH_FLAGS(SplineModifierTests, "MeshPartition.Modifier.Splines", EAu
 				});
 	}
 
-	void UpdateFalloutBoundsTest(float TargetFalloffBoundsMultiplier)
+	void ExtendSplineFalloffBoundsTest(float TargetFalloffBoundsMultiplier)
 	{
-		TSharedPtr<float> ActualTargetFalloffBoundsMultiplier = MakeShared<float>(TargetFalloffBoundsMultiplier);
+		const float SharedOriginalFalloffBoundsMultiplier = SplineModifier->GetFalloffBoundsMultiplier();
 
 		TestCommandBuilder
 			.Until(TEXT("Wait for initial modifier processing to complete"), [this]()
 				{
 					return !MeshPartitionEditorComponent->IsModifierParticipatingInActivePreviewSectionBuild(SplineModifier);
 				})
-			.Then(TEXT("Modify modifier fallout bounds"), [this, ActualTargetFalloffBoundsMultiplier]()
+			.Then(TEXT("Increase Falloff Bounds Multipler value"), [this, TargetFalloffBoundsMultiplier]()
 				{
-					FProperty* FalloffBoundsMultiplierProp = SplineModifier->GetClass()->FindPropertyByName(TEXT("FalloffBoundsMultiplier"));
-					ASSERT_THAT(IsNotNull(FalloffBoundsMultiplierProp, TEXT("Update Phase: Failed to find property FalloffBoundsMultiplier inside the modifier")));
-
-					float* FalloffBoundsMultiplierPtr = FalloffBoundsMultiplierProp->ContainerPtrToValuePtr<float>(SplineModifier);
-					ASSERT_THAT(IsNotNull(FalloffBoundsMultiplierPtr, TEXT("Update Phase: Failed to find the value of FalloffBoundsMultiplier inside the modifier")));
-
-					if (FMath::IsNearlyEqual(*FalloffBoundsMultiplierPtr, *ActualTargetFalloffBoundsMultiplier))
-					{
-						*ActualTargetFalloffBoundsMultiplier = *ActualTargetFalloffBoundsMultiplier + 1.0f;
-					}
-
-					*FalloffBoundsMultiplierPtr = *ActualTargetFalloffBoundsMultiplier;
-
-					FPropertyChangedEvent PropertyChangedEvent(FalloffBoundsMultiplierProp);
-					SplineModifier->PostEditChangeProperty(PropertyChangedEvent);
+					const FScopedTransaction Transaction(FText::FromString(TEXT("Set Falloff Bounds Multiple to an increased value")));
+					SplineModifier->Modify();
 
+					SplineModifier->SetFalloffBoundsMultiplier(TargetFalloffBoundsMultiplier);
 					SplineModifier->UpdateSplineData();
 				})
 			.Until(TEXT("Wait for modifier to finish participating in rebuild"), [this]()
 				{
 					return !MeshPartitionEditorComponent->IsModifierParticipatingInActivePreviewSectionBuild(SplineModifier);
 				})
-			.Then(TEXT("Validate the updated multiplier value on the modifier"), [this, ActualTargetFalloffBoundsMultiplier]()
+			.Then(TEXT("Validate generated mesh against ground truth"), [this]()
 				{
-					FProperty* FalloffBoundsMultiplierProp = SplineModifier->GetClass()->FindPropertyByName(TEXT("FalloffBoundsMultiplier"));
-					ASSERT_THAT(IsNotNull(FalloffBoundsMultiplierProp, TEXT("Validation Phase: Failed to find property FalloffBoundsMultiplier inside the modifier")));
-
-					float* VerifyTarget = FalloffBoundsMultiplierProp->ContainerPtrToValuePtr<float>(SplineModifier);
-					ASSERT_THAT(IsNotNull(VerifyTarget, TEXT("Validation Phase: Failed to find the value of FalloffBoundsMultiplier inside the modifier")));
-
-					ASSERT_THAT(IsTrue(FMath::IsNearlyEqual(*VerifyTarget, *ActualTargetFalloffBoundsMultiplier), TEXT("Updated FalloffBoundsMultiplier does not match the target value")));
+					const FString TestName = TEXT("Verify_Extending_Spline_Falloff_Bounds_Updated");
+					CompareToReferenceAsset(TestName);
+				})
+			.Then(TEXT("Test Undo function on the feature"), [this]()
+				{
+					GEditor->UndoTransaction();
+				})
+			.Until(TEXT("Wait for modifier update to process"), [this]()
+				{
+					return !MeshPartitionEditorComponent->IsModifierParticipatingInActivePreviewSectionBuild(SplineModifier);
+				})
+			.Then(TEXT("Validate the effects of the Undo function: Falloff Bounds Multipler value has been reverted, and Spline Falloff is being cut on the bounds"), [this, SharedOriginalFalloffBoundsMultiplier]()
+				{
+					ASSERT_THAT(IsTrue(FMath::IsNearlyEqual(SplineModifier->GetFalloffBoundsMultiplier(), SharedOriginalFalloffBoundsMultiplier), TEXT("Undo did not revert the Falloff Bounds Multipler value to the original value")));
+					const FString TestName = TEXT("Verify_Extending_Spline_Falloff_Bounds_Initial");
+					CompareToReferenceAsset(TestName);
+				})
+			.Then(TEXT("Test Redo function on the feature"), [this]()
+				{
+					GEditor->RedoTransaction();
+				})
+			.Until(TEXT("Wait for modifier update to process"), [this]()
+				{
+					return !MeshPartitionEditorComponent->IsModifierParticipatingInActivePreviewSectionBuild(SplineModifier);
+				})
+			.Then(TEXT("Validate the effects of the Redo function: Falloff Bounds Multipler value has been reapplied, and Spline Falloff smoothly reaches the ground"), [this, TargetFalloffBoundsMultiplier]()
+				{
+					ASSERT_THAT(IsTrue(FMath::IsNearlyEqual(SplineModifier->GetFalloffBoundsMultiplier(), TargetFalloffBoundsMultiplier), TEXT("Redo did not update the Falloff Bounds Multipler value to the updated value")));
+					const FString TestName = TEXT("Verify_Extending_Spline_Falloff_Bounds_Updated");
+					CompareToReferenceAsset(TestName);
 				});
 	}
 
diff --git a/Engine/Plugins/Experimental/MeshPartition/Source/MeshPartitionEditor/Public/Modifiers/MeshPartitionSplineModifier.h b/Engine/Plugins/Experimental/MeshPartition/Source/MeshPartitionEditor/Public/Modifiers/MeshPartitionSplineModifier.h
--- a/Engine/Plugins/Experimental/MeshPartition/Source/MeshPartitionEditor/Public/Modifiers/MeshPartitionSplineModifier.h
+++ b/Engine/Plugins/Experimental/MeshPartition/Source/MeshPartitionEditor/Public/Modifiers/MeshPartitionSplineModifier.h
@@ -130,6 +130,12 @@ public:
 	}
 	UE_API void SetMeshedInteriorNumTriTarget(double InMeshedInteriorNumTriTarget);
 
+	float GetFalloffBoundsMultiplier() const 
+	{
+		return FalloffBoundsMultiplier;
+	}
+	UE_API void SetFalloffBoundsMultiplier(float InFalloffBoundsMultiplier);
+
 private:
 
 	friend class FSplineSoftObjectPointerDetails;
