UE-365526

Mesh Terrain - Spline Modifier - Applying Smooth Modes for the Spline modifier interior

#virtualized
#Synced-CL 51499417

diff --git a/Engine/Plugins/Experimental/MeshPartition/Source/MeshPartitionEditor/Private/Modifiers/MeshPartitionSplineModifier.cpp b/Engine/Plugins/Experimental/MeshPartition/Source/MeshPartitionEditor/Private/Modifiers/MeshPartitionSplineModifier.cpp
--- a/Engine/Plugins/Experimental/MeshPartition/Source/MeshPartitionEditor/Private/Modifiers/MeshPartitionSplineModifier.cpp
+++ b/Engine/Plugins/Experimental/MeshPartition/Source/MeshPartitionEditor/Private/Modifiers/MeshPartitionSplineModifier.cpp
@@ -1051,4 +1051,9 @@ void USplineModifier::SetFalloffBoundsMultiplier(float InFalloffBoundsMultiplier
 
 	FalloffBoundsMultiplier = FMath::Max(1.0, InFalloffBoundsMultiplier);
 }
+void USplineModifier::SetInteriorSmoothMode(ESplineModifierInteriorSmoothMode InInteriorSmoothMode)
+{
+	InteriorSmoothMode = InInteriorSmoothMode;
+}
+
 } // namespace UE::MeshPartition
\ No newline at end of file
diff --git a/Engine/Plugins/Experimental/MeshPartition/Source/MeshPartitionEditor/Private/Test/MeshPartitionSplineModifierTests.cpp b/Engine/Plugins/Experimental/MeshPartition/Source/MeshPartitionEditor/Private/Test/MeshPartitionSplineModifierTests.cpp
--- a/Engine/Plugins/Experimental/MeshPartition/Source/MeshPartitionEditor/Private/Test/MeshPartitionSplineModifierTests.cpp
+++ b/Engine/Plugins/Experimental/MeshPartition/Source/MeshPartitionEditor/Private/Test/MeshPartitionSplineModifierTests.cpp
@@ -533,6 +533,79 @@ TEST_CLASS_WITH_FLAGS(SplineModifierTests, "MeshPartition.Modifier.Splines", EAu
 				});
 	}
 
+	// QMetry: UE - TC - 20622
+	TEST_METHOD(Verify_Applying_Smooth_Modes_For_The_Spline_Modifier_Interior)
+	{
+		TestCommandBuilder
+			.Until(TEXT("Wait for initial modifier processing to complete"), [this]()
+				{
+					return !MeshPartitionEditorComponent->IsModifierParticipatingInActivePreviewSectionBuild(SplineModifier);
+				})
+			.Then(TEXT("Set Interior Smooth Mode to Smooth"), [this]()
+				{
+					ASSERT_THAT(IsTrue(SplineModifier->GetInteriorSmoothMode() == ESplineModifierInteriorSmoothMode::Simple, TEXT("Default Interior Smooth Mode is not set to Simple")));
+
+					const FScopedTransaction Transaction(FText::FromString(TEXT("Set Interior Smooth Mode to Smooth")));
+					SplineModifier->Modify();
+
+					SplineModifier->SetInteriorSmoothMode(ESplineModifierInteriorSmoothMode::Smooth);
+					SplineModifier->UpdateSplineData();
+				})
+			.Until(TEXT("Wait for modifier update to process"), [this]()
+				{
+					return !MeshPartitionEditorComponent->IsModifierParticipatingInActivePreviewSectionBuild(SplineModifier);
+				})
+			.Then(TEXT("Validate generated mesh against ground truth: Interior Smooth Mode should be set to Smooth"), [this]()
+				{
+					ASSERT_THAT(IsTrue(SplineModifier->GetInteriorSmoothMode() == ESplineModifierInteriorSmoothMode::Smooth, TEXT("Interior Smooth Mode is not set to Smooth")));
+					const FString TestName = TEXT("Verify_Applying_Smooth_Modes_For_The_Spline_Modifier_Interior_Smooth");
+					CompareToReferenceAsset(TestName);
+				})
+			.Then(TEXT("Test Undo function on the feature"), [this]()
+				{
+					GEditor->UndoTransaction();
+				})
+			.Until(TEXT("Wait for modifier update to process"), [this]()
+				{
+					return !MeshPartitionEditorComponent->IsModifierParticipatingInActivePreviewSectionBuild(SplineModifier);
+				})
+			.Then(TEXT("Validate the effects of the Undo function: Interior Smooth Mode should be set to Simple"), [this]()
+				{
+					ASSERT_THAT(IsTrue(SplineModifier->GetInteriorSmoothMode() == ESplineModifierInteriorSmoothMode::Simple, TEXT("Interior Smooth Mode is not set to Simple")));
+					const FString TestName = TEXT("Verify_Applying_Smooth_Modes_For_The_Spline_Modifier_Interior_Simple");
+					CompareToReferenceAsset(TestName);
+				})
+			.Then(TEXT("Test Redo function on the feature"), [this]()
+				{
+					GEditor->RedoTransaction();
+				})
+			.Until(TEXT("Wait for modifier update to process"), [this]()
+				{
+					return !MeshPartitionEditorComponent->IsModifierParticipatingInActivePreviewSectionBuild(SplineModifier);
+				})
+			.Then(TEXT("Validate the effects of the Redo function: Interior Smooth Mode should be set to Smooth"), [this]()
+				{
+					ASSERT_THAT(IsTrue(SplineModifier->GetInteriorSmoothMode() == ESplineModifierInteriorSmoothMode::Smooth, TEXT("Interior Smooth Mode is not set to Smooth")));
+					const FString TestName = TEXT("Verify_Applying_Smooth_Modes_For_The_Spline_Modifier_Interior_Smooth");
+					CompareToReferenceAsset(TestName);
+				})
+			.Then(TEXT("Set Interior Smooth Mode to Detail Preserving"), [this]()
+				{
+					SplineModifier->SetInteriorSmoothMode(ESplineModifierInteriorSmoothMode::DetailPreserving);
+					SplineModifier->UpdateSplineData();
+				})
+			.Until(TEXT("Wait for modifier update to process"), [this]()
+				{
+					return !MeshPartitionEditorComponent->IsModifierParticipatingInActivePreviewSectionBuild(SplineModifier);
+				})
+			.Then(TEXT("Validate generated mesh against ground truth: Interior Smooth Mode should be set to Detail Preserving"), [this]()
+				{
+					ASSERT_THAT(IsTrue(SplineModifier->GetInteriorSmoothMode() == ESplineModifierInteriorSmoothMode::DetailPreserving, TEXT("Interior Smooth Mode is not set to Detail Preserving")));
+					const FString TestName = TEXT("Verify_Applying_Smooth_Modes_For_The_Spline_Modifier_Interior_DetailPreserving");
+					CompareToReferenceAsset(TestName);
+				});
+	}
+
 	AFTER_EACH()
 	{
 		ScopedEditorWorld.Reset();
diff --git a/Engine/Plugins/Experimental/MeshPartition/Source/MeshPartitionEditor/Public/Modifiers/MeshPartitionSplineModifier.h b/Engine/Plugins/Experimental/MeshPartition/Source/MeshPartitionEditor/Public/Modifiers/MeshPartitionSplineModifier.h
--- a/Engine/Plugins/Experimental/MeshPartition/Source/MeshPartitionEditor/Public/Modifiers/MeshPartitionSplineModifier.h
+++ b/Engine/Plugins/Experimental/MeshPartition/Source/MeshPartitionEditor/Public/Modifiers/MeshPartitionSplineModifier.h
@@ -136,6 +136,14 @@ public:
 	}
 	UE_API void SetFalloffBoundsMultiplier(float InFalloffBoundsMultiplier);
 
+	ESplineModifierInteriorSmoothMode GetInteriorSmoothMode() const
+	{
+		return InteriorSmoothMode;
+	}
+	UE_API void SetInteriorSmoothMode(ESplineModifierInteriorSmoothMode InInteriorSmoothMode);
+
+
+
 private:
 
 	friend class FSplineSoftObjectPointerDetails;
