Removing unused constructor in the MeshPartition FCachedDescriptors
Also renaming MegaMesh->MeshPartition

#Synced-CL 51539202

diff --git a/Engine/Plugins/Experimental/MeshPartition/Source/MeshPartitionEditor/Private/MeshPartitionDescriptorCache.cpp b/Engine/Plugins/Experimental/MeshPartition/Source/MeshPartitionEditor/Private/MeshPartitionDescriptorCache.cpp
--- a/Engine/Plugins/Experimental/MeshPartition/Source/MeshPartitionEditor/Private/MeshPartitionDescriptorCache.cpp
+++ b/Engine/Plugins/Experimental/MeshPartition/Source/MeshPartitionEditor/Private/MeshPartitionDescriptorCache.cpp
@@ -17,78 +17,24 @@
 
 namespace UE::MeshPartition
 {
-FCachedDescriptors::FCachedDescriptors(const AMeshPartition* InMegaMesh)
-{
-	using namespace Utils;
-
-	check(InMegaMesh != nullptr);
-	check(InMegaMesh->GetActorGuid().IsValid());
-
-	UWorld* World = InMegaMesh->GetWorld();
-	check(World != nullptr);
-	check(!World->IsPlayInEditor());		// we shouldn't be caching descriptors for PIE worlds
-
-	UWorldPartition* WorldPartition = World->GetWorldPartition();
-	check(WorldPartition != nullptr)
-
-	const UMeshPartitionDefinition* MeshPartitionDefinition = InMegaMesh->GetMeshPartitionDefinition();
-	if (MeshPartitionDefinition == nullptr)
-	{
-		// UE_LOG(LogCore, Warning, TEXT("No MegaMesh Definition provided for %s (%.6s), using default settings."), *MegaMesh->GetName(), *MegaMesh->GetActorGuid().ToString());
-		MeshPartitionDefinition = UMeshPartitionDefinition::StaticClass()->GetDefaultObject<UMeshPartitionDefinition>();
-	}
-
-	// Grab ALL MegaMeshes present in the world, and find the descriptors for all the modifiers that belong to that megamesh, as well as the current compiled sections
-	// (NOTE: This only grabs from the main world - does not include anything from inside level instances. Those are gathered during the LoadActors() step)
-	TArray<FCompiledSectionDescriptor> SectionDescs;
-	ForEachMegaMeshDescInWorldPartition(WorldPartition,
-		[InMegaMesh, this](MeshPartition::FModifierDesc& Descriptor)
-		{
-			if (Descriptor.MegaMeshGuid == InMegaMesh->GetActorGuid())
-			{
-				AllModifiersGroup.Add(Descriptor);		// not yet sorted, may add in non-optimal order
-			}
-		},
-		[InMegaMesh, &SectionDescs](FCompiledSectionDescriptor& Descriptor)
-		{	if (Descriptor.Info.MegaMeshGUID == InMegaMesh->GetActorGuid())
-			{
-				SectionDescs.Add(Descriptor);
-			}
-		});
-
-	// Sort the Modifiers by priority:
-	AllModifiersGroup.Sort(MeshPartitionDefinition->GetModifierTypePriorities());
-
-	this->MegaMesh = InMegaMesh;
-	this->MegaMeshGuid = InMegaMesh->GetActorGuid();
-	this->CompiledSections = MoveTemp(SectionDescs);
-
-	// iterate all build variants
-	TConstArrayView<MeshPartition::FCompiledSectionBuildVariant> BuildVariantArray = MeshPartitionDefinition->GetCompiledSectionBuildVariants();
-	check(BuildVariantArray.Num() > 0);
-	for (const MeshPartition::FCompiledSectionBuildVariant& BuildVariant : BuildVariantArray)
-	{
-		// Break the MegaMesh into groups
-		this->BuildVariants.Add(BuildVariant.Name, MeshPartition::FModifierGrouping(BuildVariant, *MeshPartitionDefinition, AllModifiersGroup));
-	}
-}
 
-FCachedDescriptors::FCachedDescriptors(const AMeshPartition* InMegaMesh, TConstArrayView<MeshPartition::FModifierDesc> InAllModifiersSorted, TArray<FCompiledSectionDescriptor>&& InCompiledSections)
+FCachedDescriptors::FCachedDescriptors(const AMeshPartition* InMeshPartition, TConstArrayView<MeshPartition::FModifierDesc> InAllModifiersSorted, TArray<FCompiledSectionDescriptor>&& InCompiledSections)
 {
-	check(InMegaMesh != nullptr);
-	this->MegaMeshGuid = InMegaMesh->GetActorGuid();
-	check(this->MegaMeshGuid.IsValid());
+	check(InMeshPartition != nullptr);
+	this->MeshPartitionActor = InMeshPartition;
+	this->MeshPartitionGUID = InMeshPartition->GetActorGuid();
+	check(this->MeshPartitionGUID.IsValid());
 
-	UWorld* World = InMegaMesh->GetWorld();
+	UWorld* World = InMeshPartition->GetWorld();
 	check(World != nullptr);
 
 	UWorldPartition* WorldPartition = World->GetWorldPartition();
 	check(WorldPartition != nullptr)
 
-	const UMeshPartitionDefinition* MeshPartitionDefinition = InMegaMesh->GetMeshPartitionDefinition();
+	const UMeshPartitionDefinition* MeshPartitionDefinition = InMeshPartition->GetMeshPartitionDefinition();
 	if (MeshPartitionDefinition == nullptr)
 	{
-		UE_LOG(LogCore, Warning, TEXT("No MegaMesh Definition provided for %s (%.6s), using default settings."), *InMegaMesh->GetName(), *InMegaMesh->GetActorGuid().ToString());
+		UE_LOG(LogCore, Warning, TEXT("No MeshPartition Definition provided for %s (%.6s), using default settings."), *InMeshPartition->GetName(), *InMeshPartition->GetActorGuid().ToString());
 		MeshPartitionDefinition = UMeshPartitionDefinition::StaticClass()->GetDefaultObject<UMeshPartitionDefinition>();
 	}
 
@@ -101,9 +47,6 @@ FCachedDescriptors::FCachedDescriptors(const AMeshPartition* InMegaMesh, TConstA
 	// Sort the Modifiers by priority:
 	AllModifiersGroup.Sort(MeshPartitionDefinition->GetModifierTypePriorities());
 
-	// Break the MegaMesh into groups
-	this->MegaMesh = InMegaMesh;
-	this->MegaMeshGuid = InMegaMesh->GetActorGuid();
 	this->CompiledSections = MoveTemp(InCompiledSections);
 
 	// iterate all build variants
@@ -111,7 +54,7 @@ FCachedDescriptors::FCachedDescriptors(const AMeshPartition* InMegaMesh, TConstA
 	check(BuildVariantArray.Num() > 0);
 	for (const MeshPartition::FCompiledSectionBuildVariant& BuildVariant : BuildVariantArray)
 	{
-		// Break the MegaMesh into groups
+		// break the Mesh Partition into groups, using the build variant settings
 		this->BuildVariants.Add(BuildVariant.Name, MeshPartition::FModifierGrouping(BuildVariant, *MeshPartitionDefinition, AllModifiersGroup));
 	}
 }
@@ -174,26 +117,26 @@ void FCachedDescriptors::ForAllBuildVariants(TFunctionRef<bool(const MeshPartiti
 	}
 }
 
-void FModifierDescriptorCache::AddMegaMeshToGuidMap(const AMeshPartition* MegaMesh)
+void FModifierDescriptorCache::AddMeshPartitionToGuidMap(const AMeshPartition* MeshPartition)
 {
-	if (MegaMesh == nullptr)
+	if (MeshPartition == nullptr)
 	{
 		return;
 	}
 
-	const FGuid& MegaMeshGuid = MegaMesh->GetActorGuid();
-	if (MegaMeshGuid.IsValid())
+	const FGuid& MeshPartitionGuid = MeshPartition->GetActorGuid();
+	if (MeshPartitionGuid.IsValid())
 	{
-		FSoftObjectPath MegaMeshPath(MegaMesh);
-		GuidToMegaMesh.AddUnique(MegaMeshGuid, MegaMeshPath);
+		FSoftObjectPath MeshPartitionPath(MeshPartition);
+		GuidToMeshPartitionPath.AddUnique(MeshPartitionGuid, MeshPartitionPath);
 	}
 }
 
-void FModifierDescriptorCache::UpdateCacheForAllMegaMeshesInWorld(UWorld* InWorld, TArray<FGuid> *OutAllMegaMeshGuidsInWorld)
+void FModifierDescriptorCache::UpdateCacheForAllMeshPartitionsInWorld(UWorld* InWorld, TArray<FGuid> *OutAllMeshPartitionGuidsInWorld)
 {
 	using namespace Utils;
 
-	TMap<FGuid, TTuple<TArray<MeshPartition::FModifierDesc>, TArray<FCompiledSectionDescriptor>>> MegaMeshesInWorld;
+	TMap<FGuid, TTuple<TArray<MeshPartition::FModifierDesc>, TArray<FCompiledSectionDescriptor>>> MeshPartitionsInWorld;
 	UWorldPartition* WorldPartition = InWorld->GetWorldPartition();
 	if (WorldPartition == nullptr)
 	{
@@ -210,11 +153,11 @@ void FModifierDescriptorCache::UpdateCacheForAllMegaMeshesInWorld(UWorld* InWorl
 
 	// NOTE: this only grabs actor descs from the Main world/level -- not from any level instances
 	ForEachMegaMeshDescInWorldPartition(WorldPartition,
-		[&MegaMeshesInWorld, &Modifiers, &PIEPathFixer, bIsPIE](MeshPartition::FModifierDesc& Descriptor)
+		[&MeshPartitionsInWorld, &Modifiers, &PIEPathFixer, bIsPIE](MeshPartition::FModifierDesc& Descriptor)
 		{
 			if (Descriptor.MegaMeshGuid.IsValid())
 			{
-				TArray<MeshPartition::FModifierDesc>& ModifierDescriptors = MegaMeshesInWorld.FindOrAdd(Descriptor.MegaMeshGuid).Get<0>();
+				TArray<MeshPartition::FModifierDesc>& ModifierDescriptors = MeshPartitionsInWorld.FindOrAdd(Descriptor.MegaMeshGuid).Get<0>();
 
 				// TODO: this may not correctly modify a level instanced path.. we may need to strip PIE paths from those individually?
 				FSoftObjectPath ModifierPath = Descriptor.ModifierPath;
@@ -228,84 +171,84 @@ void FModifierDescriptorCache::UpdateCacheForAllMegaMeshesInWorld(UWorld* InWorl
 				Modifiers++;
 			}
 		},
-		[&MegaMeshesInWorld, &CompiledSections](FCompiledSectionDescriptor& Descriptor)
+		[&MeshPartitionsInWorld, &CompiledSections](FCompiledSectionDescriptor& Descriptor)
 		{
-			// ignore old descriptors that don't have megamesh GUID populated
+			// ignore old descriptors that don't have mesh partition GUID populated
 			if (Descriptor.Info.MegaMeshGUID.IsValid())
 			{
-				TArray<FCompiledSectionDescriptor>& CompiledSectionDescriptors = MegaMeshesInWorld.FindOrAdd(Descriptor.Info.MegaMeshGUID).Get<1>();
+				TArray<FCompiledSectionDescriptor>& CompiledSectionDescriptors = MeshPartitionsInWorld.FindOrAdd(Descriptor.Info.MegaMeshGUID).Get<1>();
 				CompiledSectionDescriptors.Add(Descriptor);
 				CompiledSections++;
 			}
 		});
 
-	// Build the section groups for each megamesh in the world
-	// NOTE: this will not gather spatially streaming AMegaMesh actors
-	// (currently only the section actors should be streaming, not the AMegaMesh actor)
+	// Build the section groups for each mesh partition in the world
+	// NOTE: this will not gather spatially streaming AMeshPartition actors
+	// (currently only the section actors should be streaming, not the AMeshPartition actor)
 	for (TActorIterator<AMeshPartition> It(InWorld); It; ++It)
 	{
-		const AMeshPartition* MegaMesh = *It;
-		const FGuid MegaMeshGuid = MegaMesh->GetActorGuid();
+		const AMeshPartition* MeshPartition = *It;
+		const FGuid MeshPartitionGuid = MeshPartition->GetActorGuid();
 
-		// populate the GUID -> MegaMesh with the MegaMeshes that exist in the world
-		AddMegaMeshToGuidMap(MegaMesh);
+		// populate the GUID -> MeshPartition with the MeshPartitions that exist in the world
+		AddMeshPartitionToGuidMap(MeshPartition);
 
-		TTuple<TArray<MeshPartition::FModifierDesc>, TArray<FCompiledSectionDescriptor>>* Descriptors = MegaMeshesInWorld.Find(MegaMeshGuid);
+		TTuple<TArray<MeshPartition::FModifierDesc>, TArray<FCompiledSectionDescriptor>>* Descriptors = MeshPartitionsInWorld.Find(MeshPartitionGuid);
 		if ((Descriptors == nullptr) || (Descriptors->Get<0>().Num() <= 0))
 		{
-			// no descriptors found.. this megamesh will be empty
-			UE_LOG(LogMegaMeshEditor, Warning, TEXT("MegaMesh '%s' (GUID '%s') does not have any modifiers in the world, the megamesh will be empty"),
-				*FSoftObjectPath(MegaMesh).ToString(),
-				*MegaMeshGuid.ToString());
+			// no descriptors found.. this mesh partition will be empty
+			UE_LOG(LogMegaMeshEditor, Warning, TEXT("MeshPartition '%s' (GUID '%s') does not have any modifiers in the world, the MeshPartition will be empty"),
+				*FSoftObjectPath(MeshPartition).ToString(),
+				*MeshPartitionGuid.ToString());
 			continue;
 		}
 
 		// and try to build the section groups for it
-		FSoftObjectPath ObjectPath(MegaMesh);
+		FSoftObjectPath ObjectPath(MeshPartition);
 		TSharedPtr<MeshPartition::FCachedDescriptors>* CacheEntryPtr = CachedDescriptors.Find(ObjectPath);
 		TSharedPtr<MeshPartition::FCachedDescriptors> CacheEntry;
 		TConstArrayView<MeshPartition::FModifierDesc> AllModifiersSorted = Descriptors->Get<0>();
 		TArray<FCompiledSectionDescriptor>& AllCompiledSections = Descriptors->Get<1>();
 		if (CacheEntryPtr == nullptr)
 		{
-			CacheEntry = MakeShared<MeshPartition::FCachedDescriptors>(MegaMesh, AllModifiersSorted, MoveTemp(AllCompiledSections));
+			CacheEntry = MakeShared<MeshPartition::FCachedDescriptors>(MeshPartition, AllModifiersSorted, MoveTemp(AllCompiledSections));
 			CachedDescriptors.Emplace(ObjectPath, CacheEntry);
 		}
 		else
 		{
 			// TODO: we could validate and/or update the existing section...  for now we just stomp it
-			CacheEntry = MakeShared<MeshPartition::FCachedDescriptors>(MegaMesh, AllModifiersSorted, MoveTemp(AllCompiledSections));
+			CacheEntry = MakeShared<MeshPartition::FCachedDescriptors>(MeshPartition, AllModifiersSorted, MoveTemp(AllCompiledSections));
 			CachedDescriptors.Emplace(ObjectPath, CacheEntry);
 		}
 	}
 	
-	UE_LOG(LogCore, Verbose, TEXT("UpdateCacheForAllMegaMeshesInWorld (%d Modifiers, %d Compiled Sections)"), Modifiers, CompiledSections);
+	UE_LOG(LogCore, Verbose, TEXT("UpdateCacheForAllMeshPartitionsInWorld (%d Modifiers, %d Compiled Sections)"), Modifiers, CompiledSections);
 
-	if (OutAllMegaMeshGuidsInWorld != nullptr)
+	if (OutAllMeshPartitionGuidsInWorld != nullptr)
 	{
-		MegaMeshesInWorld.GetKeys(*OutAllMegaMeshGuidsInWorld);
+		MeshPartitionsInWorld.GetKeys(*OutAllMeshPartitionGuidsInWorld);
 	}
 }
 
-TSharedPtr<MeshPartition::FCachedDescriptors> FModifierDescriptorCache::GetCachedDescriptorsByPath(FSoftObjectPath MegaMeshPath)
+TSharedPtr<MeshPartition::FCachedDescriptors> FModifierDescriptorCache::GetCachedDescriptorsByPath(FSoftObjectPath MeshPartitionPath)
 {
-	if (TSharedPtr<MeshPartition::FCachedDescriptors>* CacheEntryPtr = CachedDescriptors.Find(MegaMeshPath))
+	if (TSharedPtr<MeshPartition::FCachedDescriptors>* CacheEntryPtr = CachedDescriptors.Find(MeshPartitionPath))
 	{
 		return *CacheEntryPtr;
 	}
 	return nullptr;
 }
 
-TSharedPtr<MeshPartition::FCachedDescriptors> FModifierDescriptorCache::GetCachedDescriptorsByGuid(FGuid MegaMeshGUID)
+TSharedPtr<MeshPartition::FCachedDescriptors> FModifierDescriptorCache::GetCachedDescriptorsByGuid(FGuid MeshPartitionGUID)
 {
 	TSharedPtr<MeshPartition::FCachedDescriptors>* CacheEntryPtr = nullptr;
 
-	TArray<FSoftObjectPath> MegaMeshPaths;
-	GuidToMegaMesh.MultiFind(MegaMeshGUID, MegaMeshPaths);
+	TArray<FSoftObjectPath> MeshPartitionPaths;
+	GuidToMeshPartitionPath.MultiFind(MeshPartitionGUID, MeshPartitionPaths);
 
-	for (FSoftObjectPath MegaMeshPath : MegaMeshPaths)
+	for (FSoftObjectPath MeshPartitionPath : MeshPartitionPaths)
 	{
-		CacheEntryPtr = CachedDescriptors.Find(MegaMeshPath);
+		CacheEntryPtr = CachedDescriptors.Find(MeshPartitionPath);
 		if (CacheEntryPtr != nullptr)
 		{
 			return *CacheEntryPtr;
@@ -345,20 +288,20 @@ void FModifierDescriptorCache::ForAllCachedDescriptors(TFunctionRef<bool(const F
 }
 
 TSharedPtr<MeshPartition::FCachedDescriptors> FModifierDescriptorCache::GetCachedDescriptors(
-	const AMeshPartition* MegaMesh)
+	const AMeshPartition* MeshPartition)
 {
-	check(MegaMesh);
-	FSoftObjectPath MegaMeshPath(MegaMesh);
+	check(MeshPartition);
+	FSoftObjectPath MeshPartitionPath(MeshPartition);
 	
-	if (TSharedPtr<MeshPartition::FCachedDescriptors>* CacheEntryPtr = CachedDescriptors.Find(MegaMeshPath))
+	if (TSharedPtr<MeshPartition::FCachedDescriptors>* CacheEntryPtr = CachedDescriptors.Find(MeshPartitionPath))
 	{
 		return *CacheEntryPtr;
 	}
 	else
 	{
 		// try looking it up by GUID... this will effectively grab the cached descriptors from the editor-world
-		FGuid MegaMeshGuid = MegaMesh->GetActorGuid();
-		return GetCachedDescriptorsByGuid(MegaMeshGuid);
+		FGuid MeshPartitionGuid = MeshPartition->GetActorGuid();
+		return GetCachedDescriptorsByGuid(MeshPartitionGuid);
 	}
 }
 } // namespace UE::MeshPartition
diff --git a/Engine/Plugins/Experimental/MeshPartition/Source/MeshPartitionEditor/Private/MeshPartitionSeparateWorldBuilder.cpp b/Engine/Plugins/Experimental/MeshPartition/Source/MeshPartitionEditor/Private/MeshPartitionSeparateWorldBuilder.cpp
--- a/Engine/Plugins/Experimental/MeshPartition/Source/MeshPartitionEditor/Private/MeshPartitionSeparateWorldBuilder.cpp
+++ b/Engine/Plugins/Experimental/MeshPartition/Source/MeshPartitionEditor/Private/MeshPartitionSeparateWorldBuilder.cpp
@@ -51,7 +51,7 @@ void FSeparateWorldBuilder::InitializeToBuildOutOfWorldPackage(const FString& In
 
 	if (BuilderPackage == nullptr)
 	{
-		BuilderPackage = CreatePackage(TEXT("/Temp/MegaMeshBuilderPackage"));
+		BuilderPackage = CreatePackage(TEXT("/Temp/MeshPartitionBuilderPackage"));
 	}
 	ensure(BuilderPackage);
 
@@ -83,11 +83,11 @@ void FSeparateWorldBuilder::BuildCompiledSectionMesh(MeshPartition::ACompiledSec
 
 	MeshPartition::FCompiledSectionBuildInfo BuildInfo = CompiledSection->GetBuildInfo();
 
-	// Get the megamesh definition.  Note that the compiled section is in a streaming cell, and it's parent megamesh may not be loaded
+	// Get the mesh partition definition.  Note that the compiled section is in a streaming cell, and it's parent mesh partition may not be loaded
 	// So instead we make sure to grab the definition from the path stored on the buildinfo
-	FSoftObjectPath MegaMeshDefinitionPath(BuildInfo.MegaMeshDefinitionPath);
-	TSoftObjectPtr<UMeshPartitionDefinition> MegaMeshDefinition(MegaMeshDefinitionPath);
-	UMeshPartitionDefinition* Definition = MegaMeshDefinition.Get();
+	FSoftObjectPath MeshPartitionDefinitionPath(BuildInfo.MegaMeshDefinitionPath);
+	TSoftObjectPtr<UMeshPartitionDefinition> MeshPartitionDefinition(MeshPartitionDefinitionPath);
+	UMeshPartitionDefinition* Definition = MeshPartitionDefinition.Get();
 	if (Definition == nullptr)
 	{
 		Definition = UMeshPartitionDefinition::StaticClass()->GetDefaultObject<UMeshPartitionDefinition>();
@@ -103,7 +103,7 @@ void FSeparateWorldBuilder::BuildCompiledSectionMesh(MeshPartition::ACompiledSec
 		// We search by GUID because the path can be modified by PIE
 		TSharedPtr<const MeshPartition::FCachedDescriptors> Descriptors = DescriptorCache->GetCachedDescriptorsByGuid(BuildInfo.MegaMeshGUID);
 
-		FSoftObjectPath MeshPartitionActorPath = Descriptors->GetMegaMeshPath();
+		FSoftObjectPath MeshPartitionActorPath = Descriptors->GetMeshPartitionPath();
 
 		// build the section group that corresponds to these base modifiers (after first fixing the PIE path names..)
 		MeshPartition::EditorUtils::FPIEPathFixer& PIEPathFixer = UMeshPartitionEditorSubsystem::GetPIEPathFixer();
@@ -114,7 +114,7 @@ void FSeparateWorldBuilder::BuildCompiledSectionMesh(MeshPartition::ACompiledSec
 		MeshPartition::FModifierGroup ModifierGroup = Descriptors->BuildModifierGroupForBaseModifierPaths(BuildInfo.BaseModifierPaths);
 
 		// find the mesh partition actor by GUID (NOTE: we can't find it by path, unless we use the PIE world path, otherwise it will get the mesh partition actor in the main editor world)
-		FGuid MeshPartitionActorGUID = Descriptors->GetMegaMeshGUID();
+		FGuid MeshPartitionActorGUID = Descriptors->GetMeshPartitionGUID();
 		const AMeshPartition* MeshPartitionActor = nullptr;
 		for (AMeshPartition* MeshPartition : TActorRange<AMeshPartition>(CompiledSectionWorld))
 		{
diff --git a/Engine/Plugins/Experimental/MeshPartition/Source/MeshPartitionEditor/Private/MeshPartitionWorldUpdater.cpp b/Engine/Plugins/Experimental/MeshPartition/Source/MeshPartitionEditor/Private/MeshPartitionWorldUpdater.cpp
--- a/Engine/Plugins/Experimental/MeshPartition/Source/MeshPartitionEditor/Private/MeshPartitionWorldUpdater.cpp
+++ b/Engine/Plugins/Experimental/MeshPartition/Source/MeshPartitionEditor/Private/MeshPartitionWorldUpdater.cpp
@@ -34,7 +34,7 @@ namespace UE::MeshPartition
 		// in editor worlds, update the descriptor cache immediately
 		if (InWorld->IsEditorWorld())
 		{
-			DescriptorCache->UpdateCacheForAllMegaMeshesInWorld(InWorld);
+			DescriptorCache->UpdateCacheForAllMeshPartitionsInWorld(InWorld);
 		}
 		else
 		{
@@ -60,11 +60,11 @@ namespace UE::MeshPartition
 				if (!MeshPartitionAllCachedDescriptors.IsValid())
 				{
 					UE_LOG(LogMegaMeshEditor, Verbose, TEXT("Ran without descriptors cached for MeshPartition '%s', attempting to update the cache in the current world"), *MeshPartition->GetActorNameOrLabel());
-					DescriptorCache->UpdateCacheForAllMegaMeshesInWorld(InWorld);
+					DescriptorCache->UpdateCacheForAllMeshPartitionsInWorld(InWorld);
 					MeshPartitionAllCachedDescriptors = DescriptorCache->GetCachedDescriptors(MeshPartition);
 					check(MeshPartitionAllCachedDescriptors.IsValid());
 				}
-				check(MeshPartitionAllCachedDescriptors->GetMegaMeshGUID() == MeshPartition->GetActorGuid());
+				check(MeshPartitionAllCachedDescriptors->GetMeshPartitionGUID() == MeshPartition->GetActorGuid());
 
 				TArray<FName> BuildVariantNames;
 				switch (UpdateMode)
@@ -290,7 +290,7 @@ namespace UE::MeshPartition
 		{
 			BuildInfo.BuildKey = InBuildKey;
 			BuildInfo.SetMegaMeshDefinition(&InMeshPartitionDefinition);
-			BuildInfo.MegaMeshGUID = InAllMeshPartitionDescriptors->GetMegaMeshGUID();
+			BuildInfo.MegaMeshGUID = InAllMeshPartitionDescriptors->GetMeshPartitionGUID();
 			BuildInfo.BuildVariantName = InTargetGrouping.GetBuildVariantName();
 			{	// hash the build variant
 				FDependencyHash DepdendencyHash;
@@ -298,7 +298,7 @@ namespace UE::MeshPartition
 				BuildInfo.BuildVariantHash = DepdendencyHash.GetDependentDataHash();
 			}
 			// Note: get the path from the descriptor ref, because it will be the original path (not the PIE temporary path)
-			BuildInfo.MegaMeshPath = InAllMeshPartitionDescriptors->GetMegaMeshPath();
+			BuildInfo.MegaMeshPath = InAllMeshPartitionDescriptors->GetMeshPartitionPath();
 		}
 
 		// Note: this includes compiled sections that are not for our specific build variant
@@ -306,7 +306,7 @@ namespace UE::MeshPartition
 
 		if (InAssetRegistry.IsGathering())
 		{
-			UE_LOG(LogMegaMeshEditor, Log, TEXT("AssetRegistry is not complete, forcing synchronous scan of the MegaMesh related assets"));
+			UE_LOG(LogMegaMeshEditor, Log, TEXT("AssetRegistry is not complete, forcing synchronous scan of the MeshPartition related assets"));
 
 			// force the AssetRegistry to complete scanning on the packages we care about
 			TArray<FName> PackageNamesToQuery;
@@ -513,7 +513,7 @@ namespace UE::MeshPartition
 			const FMeshPartitionUpdater& MeshPartitionUpdater = Pair.Value;
 
 			UE_LOG(LogMegaMeshEditor, Verbose, TEXT("  -- MeshPartition '%s' : %d out of date, %d non-target build variant, %d reused, %d placeholders created"),
-				*MeshPartitionUpdater.AllCachedDescriptors->GetMegaMeshPath().ToString(),
+				*MeshPartitionUpdater.AllCachedDescriptors->GetMeshPartitionPath().ToString(),
 				MeshPartitionUpdater.RemovedOutOfDate,
 				MeshPartitionUpdater.RemovedNonTargetBuildVariant,
 				MeshPartitionUpdater.ReusedCompiledSections,
@@ -631,8 +631,8 @@ namespace UE::MeshPartition
 			const FMeshPartitionUpdater& MeshPartitionUpdater = Pair.Value;
 			TSet<int32> CompiledSectionIndicesReported;
 
-			FGuid MegaMeshGuid = MeshPartitionUpdater.AllCachedDescriptors->GetMegaMeshGUID();
-			AMeshPartition** LoadedActor = LoadedMeshPartitions.Find(MegaMeshGuid);
+			FGuid MeshPartitionGuid = MeshPartitionUpdater.AllCachedDescriptors->GetMeshPartitionGUID();
+			AMeshPartition** LoadedActor = LoadedMeshPartitions.Find(MeshPartitionGuid);
 
 			FString LoadedActorPath;
 			FString LoadedActorLabel;
@@ -714,7 +714,7 @@ namespace UE::MeshPartition
 						*BuildVariantName.ToString(), SectionBaseCount, TotalBaseCount, *LoadedString, *StatusString, *DetailsString);
 				};
 
-			UE_LOG(LogMegaMeshEditor, Warning, TEXT("%s %s (GUID: %s)"), *LoadedActorLabel, *LoadedActorPath, *MegaMeshGuid.ToString());
+			UE_LOG(LogMegaMeshEditor, Warning, TEXT("%s %s (GUID: %s)"), *LoadedActorLabel, *LoadedActorPath, *MeshPartitionGuid.ToString());
 
 			// iterate variants
 			for (const TPair<FName, FBuildVariantUpdater>& Pair2 : MeshPartitionUpdater.BuildVariantUpdaters)
diff --git a/Engine/Plugins/Experimental/MeshPartition/Source/MeshPartitionEditor/Private/WorldPartitionMeshPartitionBuilder.cpp b/Engine/Plugins/Experimental/MeshPartition/Source/MeshPartitionEditor/Private/WorldPartitionMeshPartitionBuilder.cpp
--- a/Engine/Plugins/Experimental/MeshPartition/Source/MeshPartitionEditor/Private/WorldPartitionMeshPartitionBuilder.cpp
+++ b/Engine/Plugins/Experimental/MeshPartition/Source/MeshPartitionEditor/Private/WorldPartitionMeshPartitionBuilder.cpp
@@ -119,13 +119,13 @@ bool UWorldPartitionMeshPartitionBuilder::RunInternal(UWorld* InWorld, const FCe
 		return false;
 	}
 
-	TArray<FGuid> AllMegaMeshGuidsInWorld;
-	DescriptorCache->UpdateCacheForAllMegaMeshesInWorld(InWorld, &AllMegaMeshGuidsInWorld);
+	TArray<FGuid> AllMeshPartitionGuidsInWorld;
+	DescriptorCache->UpdateCacheForAllMeshPartitionsInWorld(InWorld, &AllMeshPartitionGuidsInWorld);
 
 	// The BuildKey is a unique identifier to tag all compiled sections built by this build execution (allowing us to remove old ones).
 	const FGuid BuildKey = FGuid::NewGuid();
 
-	UE_LOG(LogMegaMeshEditor, Display, TEXT("Build for World %s (%d Mesh Partitions, BuildKey %.8s)"), *InWorld->GetName(), AllMegaMeshGuidsInWorld.Num(), *BuildKey.ToString());
+	UE_LOG(LogMegaMeshEditor, Display, TEXT("Build for World %s (%d Mesh Partitions, BuildKey %.8s)"), *InWorld->GetName(), AllMeshPartitionGuidsInWorld.Num(), *BuildKey.ToString());
 
 	CompiledSectionActorDescGuidsToKeep.Empty();
 	TotalBuiltCompiledSections = 0;
@@ -136,7 +136,7 @@ bool UWorldPartitionMeshPartitionBuilder::RunInternal(UWorld* InWorld, const FCe
 	UE_LOG(LogMegaMeshEditor, Display, TEXT("[%i / %i] Initializing compiled sections build..."), 0, TotalNumberOfSectionsToBuild);
 	
 	// For each MegaMesh in the world
-	for (const FGuid& MegaMeshGUID  : AllMegaMeshGuidsInWorld)
+	for (const FGuid& MegaMeshGUID  : AllMeshPartitionGuidsInWorld)
 	{
 		// Use FWorldPartitionReference to ensure the AMeshPartition actor is loaded, then invoke BuildMegaMesh on it
 		FWorldPartitionReference MegaMeshActorRef(WorldPartition.Get(), MegaMeshGUID);
diff --git a/Engine/Plugins/Experimental/MeshPartition/Source/MeshPartitionEditor/Public/MeshPartitionDescriptorCache.h b/Engine/Plugins/Experimental/MeshPartition/Source/MeshPartitionEditor/Public/MeshPartitionDescriptorCache.h
--- a/Engine/Plugins/Experimental/MeshPartition/Source/MeshPartitionEditor/Public/MeshPartitionDescriptorCache.h
+++ b/Engine/Plugins/Experimental/MeshPartition/Source/MeshPartitionEditor/Public/MeshPartitionDescriptorCache.h
@@ -41,17 +41,20 @@ public:
 	FName GetBuildVariantName() const { return BuildVariantName; }
 };
 
+// stores all of the descriptors relevant to a specific AMeshPartition actor, as well as the FModifierGrouping for each build variant
 struct FCachedDescriptors
 {
 	// NOTE: this struct is copy-on-write so it is thread-safe after construction (as long as you hold a TSharedPtr to it)
 
 public:
-	FCachedDescriptors(const AMeshPartition* InMegaMesh);
-	FCachedDescriptors(const AMeshPartition* InMegaMesh, TConstArrayView<MeshPartition::FModifierDesc> InAllModifiersSorted, TArray<FCompiledSectionDescriptor>&& InCompiledSections);
+	FCachedDescriptors(const AMeshPartition* InMeshPartition, TConstArrayView<MeshPartition::FModifierDesc> InAllModifiersSorted, TArray<FCompiledSectionDescriptor>&& InCompiledSections);
 
 private:
-	FGuid MegaMeshGuid;
-	TSoftObjectPtr<const AMeshPartition> MegaMesh;
+	// actor GUID of the mesh partition actor
+	FGuid MeshPartitionGUID;
+
+	// object path of the mesh partition actor
+	TSoftObjectPtr<const AMeshPartition> MeshPartitionActor;
 
 	// group containing all existing modifiers (including both base and non-base modifiers)
 	MeshPartition::FModifierGroup AllModifiersGroup;
@@ -99,14 +102,19 @@ public:
 		return TConstArrayView<MeshPartition::FModifierGroup>();
 	}
 
-	FSoftObjectPath GetMegaMeshPath() const
+	const FSoftObjectPath& GetMeshPartitionPath() const
+	{
+		return MeshPartitionActor.GetUniqueID();
+	}
+
+	const FGuid& GetMeshPartitionGUID() const
 	{
-		return MegaMesh.GetUniqueID();
+		return MeshPartitionGUID;
 	}
 
-	FGuid GetMegaMeshGUID() const
+	const AMeshPartition* GetMeshPartitionActor() const
 	{
-		return MegaMeshGuid;
+		return MeshPartitionActor.Get();
 	}
 
 	void ForAllBuildVariants(TFunctionRef<bool(const MeshPartition::FModifierGrouping&)> InFunc) const;
@@ -115,22 +123,22 @@ public:
 struct FModifierDescriptorCache
 {
 public:
-	// returns an array of the mega mesh guids that exist in the world
-	void UpdateCacheForAllMegaMeshesInWorld(UWorld* World, TArray<FGuid>* OutAllMegaMeshGuidsInWorld = nullptr);
+	// returns an array of the mesh partition guids that exist in the world
+	void UpdateCacheForAllMeshPartitionsInWorld(UWorld* World, TArray<FGuid>* OutAllMeshPartitionGuidsInWorld = nullptr);
 
-	TSharedPtr<MeshPartition::FCachedDescriptors> GetCachedDescriptors(const AMeshPartition* MegaMesh);
-	TSharedPtr<MeshPartition::FCachedDescriptors> GetCachedDescriptorsByPath(FSoftObjectPath MegaMeshPath);
-	TSharedPtr<MeshPartition::FCachedDescriptors> GetCachedDescriptorsByGuid(FGuid MegaMeshGUID);
+	TSharedPtr<MeshPartition::FCachedDescriptors> GetCachedDescriptors(const AMeshPartition* MeshPartition);
+	TSharedPtr<MeshPartition::FCachedDescriptors> GetCachedDescriptorsByPath(FSoftObjectPath MeshPartitionPath);
+	TSharedPtr<MeshPartition::FCachedDescriptors> GetCachedDescriptorsByGuid(FGuid MeshPartitionGUID);
 
 	void ForAllCachedDescriptors(TFunctionRef<bool(const FSoftObjectPath&, const TSharedPtr<MeshPartition::FCachedDescriptors>&)> InFunc) const;
 
 private:
-	void AddMegaMeshToGuidMap(const AMeshPartition* MegaMesh);
+	void AddMeshPartitionToGuidMap(const AMeshPartition* MeshPartition);
 
 private:
-	// Mapping from MegaMeshGuid (ActorGuid of AMeshPartition actor) to the soft object (path) pointer to the AMeshPartition.  Populated by cached megameshes.
+	// Mapping from MeshPartitionGuid (ActorGuid of AMeshPartition actor) to the soft object (path) pointer to the AMeshPartition.  Populated by cached mesh partitions.
 	// This can have multiple entries -- for example when a world is duplicated, like when PIE, the same GUID maps to a new object path.
-	TMultiMap<FGuid, FSoftObjectPath> GuidToMegaMesh;
+	TMultiMap<FGuid, FSoftObjectPath> GuidToMeshPartitionPath;
 
 	// mapping storing the cached descriptors for each mega mesh actor
 	TMap<FSoftObjectPath, TSharedPtr<MeshPartition::FCachedDescriptors>> CachedDescriptors;
