add option to filter remesh modifier to only affect triangles w/ vertex density weights above a weight threshold

#Synced-CL 51646291

diff --git a/Engine/Plugins/Experimental/MeshPartition/Source/MeshPartitionEditor/Private/Modifiers/MeshPartitionRemeshModifier.cpp b/Engine/Plugins/Experimental/MeshPartition/Source/MeshPartitionEditor/Private/Modifiers/MeshPartitionRemeshModifier.cpp
--- a/Engine/Plugins/Experimental/MeshPartition/Source/MeshPartitionEditor/Private/Modifiers/MeshPartitionRemeshModifier.cpp
+++ b/Engine/Plugins/Experimental/MeshPartition/Source/MeshPartitionEditor/Private/Modifiers/MeshPartitionRemeshModifier.cpp
@@ -237,6 +237,8 @@ void URemeshModifierBase::GatherDependencies(MeshPartition::IDependencyInterface
 	Dependencies += bUseDensityWeightChannel;
 	Dependencies += DensityWeightChannelName;
 	Dependencies += RelativeDensity;
+	Dependencies += bUseWeightThreshold;
+	Dependencies += MinWeightThreshold;
 
 	Dependencies += TessellationMethod;
 	Dependencies += bUseTargetEdgeLength;
@@ -328,6 +330,8 @@ TSharedPtr<const MeshPartition::IModifierBackgroundOp> URemeshModifier::CreateBa
 		Op->bUseDensityWeightChannel = bUseDensityWeightChannel;
 		Op->DensityWeightChannelName = DensityWeightChannelName;
 		Op->RelativeDensity = RelativeDensity;
+		Op->bUseWeightThreshold = bUseWeightThreshold;
+		Op->MinWeightThreshold = MinWeightThreshold;
 
 		OutOp = Op;
 	}
@@ -352,6 +356,8 @@ TSharedPtr<const MeshPartition::IModifierBackgroundOp> URemeshModifier::CreateBa
 		Op->bUseDensityWeightChannel = bUseDensityWeightChannel;
 		Op->DensityWeightChannelName = DensityWeightChannelName;
 		Op->RelativeDensity = RelativeDensity;
+		Op->bUseWeightThreshold = bUseWeightThreshold;
+		Op->MinWeightThreshold = MinWeightThreshold;
 
 		OutOp = Op;
 	}
diff --git a/Engine/Plugins/Experimental/MeshPartition/Source/MeshPartitionEditor/Private/Modifiers/Ops/MeshPartitionRemeshOp.cpp b/Engine/Plugins/Experimental/MeshPartition/Source/MeshPartitionEditor/Private/Modifiers/Ops/MeshPartitionRemeshOp.cpp
--- a/Engine/Plugins/Experimental/MeshPartition/Source/MeshPartitionEditor/Private/Modifiers/Ops/MeshPartitionRemeshOp.cpp
+++ b/Engine/Plugins/Experimental/MeshPartition/Source/MeshPartitionEditor/Private/Modifiers/Ops/MeshPartitionRemeshOp.cpp
@@ -274,6 +274,27 @@ namespace UE::MeshPartition
 			VertInCoverage.Add(Vid, bInCoverage);
 			return bInCoverage;
 		};
+		const FDynamicMeshWeightAttribute* DensityWeightLayer = nullptr;
+		if (bUseDensityWeightChannel)
+		{
+			if (const FDynamicMeshAttributeSet* const Attribs = Submesh.Attributes())
+			{
+				for (int32 Idx = 0; Idx < Attribs->NumWeightLayers(); ++Idx)
+				{
+					if (Attribs->GetWeightLayer(Idx)->GetName() == DensityWeightChannelName)
+					{
+						DensityWeightLayer = Attribs->GetWeightLayer(Idx);
+					}
+				}
+			}
+		}
+		// helper lambda to be called only if the weight layer is in use
+		auto VertAtBelowThreshold = [this, DensityWeightLayer](int32 VID)
+		{
+			float Value;
+			DensityWeightLayer->GetValue(VID, &Value);
+			return (Value <= MinWeightThreshold);
+		};
 
 		for (int32 Tid : Submesh.TriangleIndicesItr())
 		{
@@ -283,6 +304,14 @@ namespace UE::MeshPartition
 				// One of the verts is not inside coverage
 				continue;
 			}
+			// if we have a weight threshold and all weights are at/below threshold, skip the tri
+			if (bUseWeightThreshold && DensityWeightLayer &&
+				VertAtBelowThreshold(TriVids.A) &&
+				VertAtBelowThreshold(TriVids.B) &&
+				VertAtBelowThreshold(TriVids.C))
+			{
+				continue;
+			}
 
 			TidROI.Add(Tid);
 		}
diff --git a/Engine/Plugins/Experimental/MeshPartition/Source/MeshPartitionEditor/Private/Modifiers/Ops/MeshPartitionRemeshOp.h b/Engine/Plugins/Experimental/MeshPartition/Source/MeshPartitionEditor/Private/Modifiers/Ops/MeshPartitionRemeshOp.h
--- a/Engine/Plugins/Experimental/MeshPartition/Source/MeshPartitionEditor/Private/Modifiers/Ops/MeshPartitionRemeshOp.h
+++ b/Engine/Plugins/Experimental/MeshPartition/Source/MeshPartitionEditor/Private/Modifiers/Ops/MeshPartitionRemeshOp.h
@@ -33,6 +33,8 @@ namespace UE::MeshPartition
 		bool bUseDensityWeightChannel = false;
 		MeshPartition::FChannelName DensityWeightChannelName;
 		float RelativeDensity;
+		bool bUseWeightThreshold = false;
+		float MinWeightThreshold = 0.f;
 
 	protected:
 
@@ -49,7 +51,7 @@ namespace UE::MeshPartition
 		// Generate a new random guid before submitting any code changes to the op
 		static FGuid GetCodeVersionKey()
 		{
-			static FGuid VersionKey(TEXT("32385a18-2c4c-4520-b793-08863e786a1d"));
+			static FGuid VersionKey(TEXT("20d6f04f-c4e8-41e4-8479-d621a86ee652"));
 			return VersionKey;
 		}
 
diff --git a/Engine/Plugins/Experimental/MeshPartition/Source/MeshPartitionEditor/Private/Modifiers/Ops/MeshPartitionTessellateOp.cpp b/Engine/Plugins/Experimental/MeshPartition/Source/MeshPartitionEditor/Private/Modifiers/Ops/MeshPartitionTessellateOp.cpp
--- a/Engine/Plugins/Experimental/MeshPartition/Source/MeshPartitionEditor/Private/Modifiers/Ops/MeshPartitionTessellateOp.cpp
+++ b/Engine/Plugins/Experimental/MeshPartition/Source/MeshPartitionEditor/Private/Modifiers/Ops/MeshPartitionTessellateOp.cpp
@@ -290,12 +290,37 @@ void FMegaMeshTessellateBackgroundOp::ApplyModifications(MeshPartition::FMeshVie
 
 	FDynamicMesh3& SubMesh = InMeshView.GetSubmeshMutable();
 
+	const FDynamicMeshWeightAttribute* DensityWeightLayer = nullptr;
+	if (bUseDensityWeightChannel)
+	{
+		if (const FDynamicMeshAttributeSet* const Attribs = SubMesh.Attributes())
+		{
+			for (int32 Idx = 0; Idx < Attribs->NumWeightLayers(); ++Idx)
+			{
+				if (Attribs->GetWeightLayer(Idx)->GetName() == DensityWeightChannelName)
+				{
+					DensityWeightLayer = Attribs->GetWeightLayer(Idx);
+				}
+			}
+		}
+	}
+
 	TSet<int32> InsideTriangles;
-	for (int32 VertexID = 0; VertexID < SubMesh.VertexCount(); ++VertexID)
+	for (int32 VertexID : SubMesh.VertexIndicesItr())
 	{
 		const FVector3d LocalPosition = ModifierTransform.InverseTransformPosition(InTransform.TransformPosition(SubMesh.GetVertex(VertexID)));
 		if (LocalBounds.IsInside(LocalPosition))
 		{
+			// don't add vertex if weight is at/below threshold
+			if (bUseWeightThreshold && DensityWeightLayer)
+			{
+				float Value;
+				DensityWeightLayer->GetValue(VertexID, &Value);
+				if (Value <= MinWeightThreshold)
+				{
+					continue;
+				}
+			}
 			for (const int32 TriIndex : SubMesh.VtxTrianglesItr(VertexID))
 			{
 				InsideTriangles.Add(TriIndex);
diff --git a/Engine/Plugins/Experimental/MeshPartition/Source/MeshPartitionEditor/Private/Modifiers/Ops/MeshPartitionTessellateOp.h b/Engine/Plugins/Experimental/MeshPartition/Source/MeshPartitionEditor/Private/Modifiers/Ops/MeshPartitionTessellateOp.h
--- a/Engine/Plugins/Experimental/MeshPartition/Source/MeshPartitionEditor/Private/Modifiers/Ops/MeshPartitionTessellateOp.h
+++ b/Engine/Plugins/Experimental/MeshPartition/Source/MeshPartitionEditor/Private/Modifiers/Ops/MeshPartitionTessellateOp.h
@@ -26,6 +26,8 @@ public:
 	bool bUseDensityWeightChannel = false;
 	MeshPartition::FChannelName DensityWeightChannelName;
 	float RelativeDensity;
+	bool bUseWeightThreshold = false;
+	float MinWeightThreshold = 0.f;
 
 	bool bVertexSmoothing;
 	double SmoothingStrength;
@@ -52,7 +54,7 @@ public:
 	// Generate a new random guid before submitting any code changes to the op
 	static FGuid GetCodeVersionKey()
 	{
-		static FGuid VersionKey(TEXT("f5483f27-0e8e-4483-bc02-f14f471b41e6"));
+		static FGuid VersionKey(TEXT("6d90ac8e-ffb2-4d02-a952-f618ac512f60"));
 		return VersionKey;
 	}
 
diff --git a/Engine/Plugins/Experimental/MeshPartition/Source/MeshPartitionEditor/Public/Modifiers/MeshPartitionRemeshModifier.h b/Engine/Plugins/Experimental/MeshPartition/Source/MeshPartitionEditor/Public/Modifiers/MeshPartitionRemeshModifier.h
--- a/Engine/Plugins/Experimental/MeshPartition/Source/MeshPartitionEditor/Public/Modifiers/MeshPartitionRemeshModifier.h
+++ b/Engine/Plugins/Experimental/MeshPartition/Source/MeshPartitionEditor/Public/Modifiers/MeshPartitionRemeshModifier.h
@@ -254,6 +254,14 @@ protected:
 	UPROPERTY(EditAnywhere, Category = DensityWeightChannel, meta = (EditCondition = "CurrentOperation == ERemeshModifierOperation::Remesh || bUseTargetEdgeLength", UIMin = "-2.0", UIMax = "2.0"))
 	float RelativeDensity = 0.f;
 
+	// Whether to only change the mesh where weight values are above the MinWeightThreshold
+	UPROPERTY(EditAnywhere, Category = DensityWeightChannel, meta = (EditCondition = "CurrentOperation == ERemeshModifierOperation::Remesh || bUseTargetEdgeLength"))
+	bool bUseWeightThreshold = false;
+
+	// Min weight threshold below which the mesh will not be changed
+	UPROPERTY(EditAnywhere, Category = DensityWeightChannel, meta = (EditCondition = "bUseWeightThreshold && (CurrentOperation == ERemeshModifierOperation::Remesh || bUseTargetEdgeLength)", UIMin = "0.0", UIMax = "1.0"))
+	float MinWeightThreshold = 0.01f;
+
 	/** When true, draws a box showing affected volume in editor */
 	UPROPERTY(EditAnywhere, Category = "Modifier|ModifierSettings", meta = (DisplayName = "Local Bounds"))
 	bool bDrawAffectedBox = true;
