Renaming MegaMesh -> MeshPartition in the world partition mesh partition builder (cmd line build)

#Synced-CL 51539202

diff --git a/Engine/Plugins/Experimental/MeshPartition/Source/MeshPartitionEditor/Private/MeshPartitionEditorModule.cpp b/Engine/Plugins/Experimental/MeshPartition/Source/MeshPartitionEditor/Private/MeshPartitionEditorModule.cpp
--- a/Engine/Plugins/Experimental/MeshPartition/Source/MeshPartitionEditor/Private/MeshPartitionEditorModule.cpp
+++ b/Engine/Plugins/Experimental/MeshPartition/Source/MeshPartitionEditor/Private/MeshPartitionEditorModule.cpp
@@ -31,11 +31,11 @@ void FMeshPartitionEditorModule::StartupModule()
 	const FText SectionLabel = LOCTEXT("MeshPartition", "Mesh Partition");
 
 	FEditorBuildUtils::RegisterCustomBuildType(EditorBuildType,
-											FCanDoEditorBuildDelegate::CreateStatic(&UWorldPartitionMeshPartitionBuilder::CanBuildMegaMeshes),
+											FCanDoEditorBuildDelegate::CreateStatic(&UWorldPartitionMeshPartitionBuilder::CanBuildMeshPartitions),
 											FDoEditorBuildDelegate::CreateLambda(
 												[](UWorld* InWorld, FName InBuildOption)
 												{
-													return UWorldPartitionMeshPartitionBuilder::BuildMegaMeshes(InWorld, EMeshPartitionReuseMode::ByPackageHash);
+													return UWorldPartitionMeshPartitionBuilder::BuildMeshPartitions(InWorld, EMeshPartitionReuseMode::ByPackageHash);
 												}),
 											/*BuildAllExtensionPoint*/NAME_None,
 											/*MenuEntryLabel*/LOCTEXT("BuildMeshPartition", "Build Mesh Partition"),
@@ -43,11 +43,11 @@ void FMeshPartitionEditorModule::StartupModule()
 											/*bExternalProcess=*/true);
 
 	FEditorBuildUtils::RegisterCustomBuildType(EditorBuildTypeForce,
-												FCanDoEditorBuildDelegate::CreateStatic(&UWorldPartitionMeshPartitionBuilder::CanBuildMegaMeshes),
+												FCanDoEditorBuildDelegate::CreateStatic(&UWorldPartitionMeshPartitionBuilder::CanBuildMeshPartitions),
 												FDoEditorBuildDelegate::CreateLambda(
 												[](UWorld* InWorld, FName InBuildOption)
 												{
-													return UWorldPartitionMeshPartitionBuilder::BuildMegaMeshes(InWorld, EMeshPartitionReuseMode::ForceRebuild);
+													return UWorldPartitionMeshPartitionBuilder::BuildMeshPartitions(InWorld, EMeshPartitionReuseMode::ForceRebuild);
 												}),
 												/*BuildAllExtensionPoint*/NAME_None,
 												/*MenuEntryLabel*/LOCTEXT("BuildMeshPartitionForce", "Build Mesh Partition (force rebuild)"),
diff --git a/Engine/Plugins/Experimental/MeshPartition/Source/MeshPartitionEditor/Private/WorldPartitionMeshPartitionBuilder.cpp b/Engine/Plugins/Experimental/MeshPartition/Source/MeshPartitionEditor/Private/WorldPartitionMeshPartitionBuilder.cpp
--- a/Engine/Plugins/Experimental/MeshPartition/Source/MeshPartitionEditor/Private/WorldPartitionMeshPartitionBuilder.cpp
+++ b/Engine/Plugins/Experimental/MeshPartition/Source/MeshPartitionEditor/Private/WorldPartitionMeshPartitionBuilder.cpp
@@ -48,17 +48,17 @@
 #include "WorldPartition/DataLayer/DataLayerManager.h"
 #include "WorldPartition/DataLayer/DataLayerInstance.h"
 
-#define LOCTEXT_NAMESPACE "MegaMeshBuilder"
+#define LOCTEXT_NAMESPACE "MeshPartitionBuilder"
 
-static FAutoConsoleVariable CVarMegaMeshBuildCommandWaitForDebugger(
-	TEXT("megamesh.Build.WaitForDebugger"),
+static FAutoConsoleVariable CVarMeshPartitionMeshBuildCommandWaitForDebugger(
+	TEXT("MeshPartition.Build.WaitForDebugger"),
 	0,
 	TEXT("Enables wait for debugger when starting a Mesh Partition build process"));
 
 namespace UE::MeshPartition
 {
 
-namespace WorldPartitionMegaMeshBuilderLocals
+namespace WorldPartitionMeshPartitionBuilderLocals
 {
 	int32 GetTotalNumberOfSectionsToBuild(MeshPartition::FModifierDescriptorCache* InDescriptorCache)
 	{
@@ -131,20 +131,24 @@ bool UWorldPartitionMeshPartitionBuilder::RunInternal(UWorld* InWorld, const FCe
 	TotalBuiltCompiledSections = 0;
 	TotalFailedCompiledSections = 0;
 	TotalReusedCompiledSections = 0;
-	TotalNumberOfSectionsToBuild = WorldPartitionMegaMeshBuilderLocals::GetTotalNumberOfSectionsToBuild(DescriptorCache);
+	TotalNumberOfSectionsToBuild = WorldPartitionMeshPartitionBuilderLocals::GetTotalNumberOfSectionsToBuild(DescriptorCache);
 
 	UE_LOG(LogMegaMeshEditor, Display, TEXT("[%i / %i] Initializing compiled sections build..."), 0, TotalNumberOfSectionsToBuild);
 	
-	// For each MegaMesh in the world
-	for (const FGuid& MegaMeshGUID  : AllMeshPartitionGuidsInWorld)
+	// For each MeshPartition in the world
+	int32 FailedMeshPartitionBuilds = 0;
+	for (const FGuid& MeshPartitionGUID  : AllMeshPartitionGuidsInWorld)
 	{
-		// Use FWorldPartitionReference to ensure the AMeshPartition actor is loaded, then invoke BuildMegaMesh on it
-		FWorldPartitionReference MegaMeshActorRef(WorldPartition.Get(), MegaMeshGUID);
+		// Use FWorldPartitionReference to ensure the AMeshPartition actor is loaded, then invoke BuildMeshPartition on it
+		FWorldPartitionReference MeshPartitionActorRef(WorldPartition.Get(), MeshPartitionGUID);
 
-		// Note: Not every MegaMesh GUID references a valid mesh partition actor -- for example, there may be stale compiled sections that reference the GUID of old mesh partitions that have since been deleted
-		if (const AMeshPartition* MegaMesh = Cast<AMeshPartition>(MegaMeshActorRef.GetActor()))
+		// Note: Not every MeshPartition GUID references a valid mesh partition actor -- for example, there may be stale compiled sections that reference the GUID of old mesh partitions that have since been deleted
+		if (const AMeshPartition* MeshPartition = Cast<AMeshPartition>(MeshPartitionActorRef.GetActor()))
 		{
-			bool bResult = BuildMegaMesh(MegaMesh, BuildKey, *DescriptorCache, InPackageHelper);
+			if (!BuildMeshPartition(MeshPartition, BuildKey, *DescriptorCache, InPackageHelper))
+			{
+				FailedMeshPartitionBuilds++;
+			}
 		}
 	}
 
@@ -245,44 +249,63 @@ bool UWorldPartitionMeshPartitionBuilder::RunInternal(UWorld* InWorld, const FCe
 		TotalFailedCompiledSections,
 		TotalRemovedCompiledSections);
 
-	return true;
+	if ((FailedMeshPartitionBuilds == 0) && (TotalFailedCompiledSections == 0))
+	{
+		return true;
+	}
+
+	UE_LOG(LogMegaMeshEditor, Warning, TEXT("Failed to build for World %s (Built %d, Reused %d, Failed %d, Removed %d Compiled Sections), %d Mesh Partitions Failed"), *InWorld->GetName(),
+		TotalBuiltCompiledSections,
+		TotalReusedCompiledSections,
+		TotalFailedCompiledSections,
+		TotalRemovedCompiledSections,
+		FailedMeshPartitionBuilds);
+
+	return false;
 }
 
-bool UWorldPartitionMeshPartitionBuilder::BuildMegaMesh(const AMeshPartition* MegaMesh, const FGuid& BuildKey, MeshPartition::FModifierDescriptorCache& DescriptorCache, FPackageSourceControlHelper& InPackageHelper)
+bool UWorldPartitionMeshPartitionBuilder::BuildMeshPartition(const AMeshPartition* MeshPartition, const FGuid& BuildKey, MeshPartition::FModifierDescriptorCache& DescriptorCache, FPackageSourceControlHelper& InPackageHelper)
 {
-	TRACE_CPUPROFILER_EVENT_SCOPE(UWorldPartitionMeshPartitionBuilder::BuildMegaMesh);
+	TRACE_CPUPROFILER_EVENT_SCOPE(UWorldPartitionMeshPartitionBuilder::BuildMeshPartition);
+
+	if (!ensure(MeshPartition))
+	{
+		UE_LOG(LogMegaMeshEditor, Error, TEXT("Cannot build null MeshPartition actor reference."));
+		return false;
+	}
 
-	const FGuid& MegaMeshGUID = MegaMesh->GetActorGuid();
-	if (!MegaMeshGUID.IsValid())
+	const FGuid& MeshPartitionGUID = MeshPartition->GetActorGuid();
+	if (!ensure(MeshPartitionGUID.IsValid()))
 	{
+		UE_LOG(LogMegaMeshEditor, Error, TEXT("Invalid MeshPartitionGUID for Mesh Partition '%s'"), *MeshPartition->GetActorNameOrLabel());
 		return false;
 	}
 
-	UWorld* World = MegaMesh->GetWorld();
+	UWorld* World = MeshPartition->GetWorld();
 	check(World);
 
 	// Create a new build info -- this gets recorded on all compiled sections produced here.
 	MeshPartition::FCompiledSectionBuildInfo BuildInfo;
 	BuildInfo.BuildKey = BuildKey;
-	BuildInfo.MegaMeshGUID = MegaMeshGUID;
-	BuildInfo.MegaMeshPath = MegaMesh;
+	BuildInfo.MegaMeshGUID = MeshPartitionGUID;
+	BuildInfo.MegaMeshPath = MeshPartition;
 
-	TSharedPtr<const MeshPartition::FCachedDescriptors> Descriptors = DescriptorCache.GetCachedDescriptors(MegaMesh);
+	TSharedPtr<const MeshPartition::FCachedDescriptors> Descriptors = DescriptorCache.GetCachedDescriptors(MeshPartition);
 	if (!Descriptors.IsValid())
 	{
-		UE_LOG(LogMegaMeshEditor, Warning, TEXT("Could not gather descriptors for the Mesh Partition '%s'"), *MegaMesh->GetActorNameOrLabel());
+		UE_LOG(LogMegaMeshEditor, Warning, TEXT("Could not gather descriptors for the Mesh Partition '%s'"), *MeshPartition->GetActorNameOrLabel());
 		return false;
 	}
 	TArrayView<const FCompiledSectionDescriptor> OldCompiledSections = Descriptors->GetCompiledSections();
 
 	// Use FWorldPartitionReference to load/unload everything in scope
-	FWorldPartitionReference MegaMeshActorRef(WorldPartition.Get(), MegaMeshGUID);
-	check(MegaMesh == CastChecked<AMeshPartition>(MegaMeshActorRef.GetActor()));
+	FWorldPartitionReference MeshPartitionActorRef(WorldPartition.Get(), MeshPartitionGUID);
+	check(MeshPartition == CastChecked<AMeshPartition>(MeshPartitionActorRef.GetActor()));
 
-	const UMeshPartitionDefinition* MeshPartitionDefinition = MegaMesh->GetMeshPartitionDefinition();
+	const UMeshPartitionDefinition* MeshPartitionDefinition = MeshPartition->GetMeshPartitionDefinition();
 	if (MeshPartitionDefinition == nullptr)
 	{
-		UE_LOG(LogMegaMeshEditor, Verbose, TEXT("No Mesh Partition Definition provided for %s (%.6s), using default settings."), *MegaMesh->GetActorNameOrLabel(), *MegaMeshGUID.ToString());
+		UE_LOG(LogMegaMeshEditor, Verbose, TEXT("No Mesh Partition Definition provided for %s (%.6s), using default settings."), *MeshPartition->GetActorNameOrLabel(), *MeshPartitionGUID.ToString());
 		MeshPartitionDefinition = UMeshPartitionDefinition::StaticClass()->GetDefaultObject<UMeshPartitionDefinition>();
 	}
 	BuildInfo.SetMegaMeshDefinition(MeshPartitionDefinition);
@@ -290,15 +313,16 @@ bool UWorldPartitionMeshPartitionBuilder::BuildMegaMesh(const AMeshPartition* Me
 	TConstArrayView<MeshPartition::FCompiledSectionBuildVariant> BuildVariants = MeshPartitionDefinition->GetCompiledSectionBuildVariants();
 	check(BuildVariants.Num() > 0);
 
-	UMeshPartitionEditorComponent* MegaMeshEditorComponent = Cast<UMeshPartitionEditorComponent>(MegaMesh->GetMeshPartitionComponent());
+	UMeshPartitionEditorComponent* MeshPartitionEditorComponent = Cast<UMeshPartitionEditorComponent>(MeshPartition->GetMeshPartitionComponent());
+	check(MeshPartitionEditorComponent != nullptr);
+
 	// avoid level streaming triggering a preview section build while we compiled section resources are being built.
-	MegaMeshEditorComponent->SetPreviewSectionBuildEnabled(false);
+	MeshPartitionEditorComponent->SetPreviewSectionBuildEnabled(false);
 
-	const FTransform& Transform = MegaMesh->GetTransform();
-	check(MegaMeshEditorComponent != nullptr);
+	const FTransform& Transform = MeshPartition->GetTransform();
 
 	UE_LOG(LogMegaMeshEditor, Display, TEXT("-- %.6s Build Mesh Partition %s (%d Variants, %d Existing Compiled Sections"),
-		*MegaMeshGUID.ToString(), *MegaMesh->GetActorNameOrLabel(),
+		*MeshPartitionGUID.ToString(), *MeshPartition->GetActorNameOrLabel(),
 		BuildVariants.Num(), OldCompiledSections.Num());
 
 	if (AMeshPartitionDataLayerContainer* Container = AMeshPartitionDataLayerContainer::GetOrCreate(World))
@@ -316,11 +340,11 @@ bool UWorldPartitionMeshPartitionBuilder::BuildMegaMesh(const AMeshPartition* Me
 		SavePackages(PackagesToSave, InPackageHelper);
 	}
 
-	// Iterate over all the BuildVariants of the MegaMesh
+	// Iterate over all the BuildVariants of the MeshPartition
 	for (const MeshPartition::FCompiledSectionBuildVariant& BuildVariant : BuildVariants)
 	{
 		TRACE_CPUPROFILER_EVENT_SCOPE_TEXT(
-			*FString::Format(TEXT("UWorldPartitionMeshPartitionBuilder::BuildMegaMeshVariant {0}"), {BuildVariant.Name.ToString()}));
+			*FString::Format(TEXT("UWorldPartitionMeshPartitionBuilder::BuildMeshPartitionVariant {0}"), {BuildVariant.Name.ToString()}));
 
 		// Record the build variant, so at cook time we can filter to the proper variant for each platform
 		BuildInfo.BuildVariantName = BuildVariant.Name;
@@ -329,11 +353,11 @@ bool UWorldPartitionMeshPartitionBuilder::BuildMegaMesh(const AMeshPartition* Me
 		BuildInfo.GetMegaMeshDefinition()->GatherDependencies(DependencyHash, BuildVariant.Name);
 		BuildInfo.BuildVariantHash = DependencyHash.GetDependentDataHash();
 
-		UE_LOG(LogMegaMeshEditor, Display, TEXT("-- %.6s -- %s Build Variant"), *MegaMeshGUID.ToString(), *BuildVariant.Name.ToString());
+		UE_LOG(LogMegaMeshEditor, Display, TEXT("-- %.6s -- %s Build Variant"), *MeshPartitionGUID.ToString(), *BuildVariant.Name.ToString());
 
 		TArrayView<const MeshPartition::FModifierGroup> ModifierGroups = Descriptors->GetModifierGroups(BuildVariant);
 
-		MegaMeshEditorComponent->ResetGroupRegistry(ModifierGroups);
+		MeshPartitionEditorComponent->ResetGroupRegistry(ModifierGroups);
 
 		MeshPartition::FBuilderSettings Settings;
 		Settings.BuildType = MeshPartition::EBuildType::CompiledSection;
@@ -355,12 +379,12 @@ bool UWorldPartitionMeshPartitionBuilder::BuildMegaMesh(const AMeshPartition* Me
 			BuildInfo.BaseModifierPaths.Empty();
 			Algo::Transform(ModifierGroup.BaseDescs(), BuildInfo.BaseModifierPaths, [](const MeshPartition::FModifierDesc& BaseDescriptor) { return BaseDescriptor.ModifierPath; });
 
-			FModifierGroupTask* ModifierGroupTask = new FModifierGroupTask(ModifierGroup, Settings, BuildInfo, BuildVariant, MegaMeshEditorComponent, *this, OldCompiledSections);
+			FModifierGroupTask* ModifierGroupTask = new FModifierGroupTask(ModifierGroup, Settings, BuildInfo, BuildVariant, MeshPartitionEditorComponent, *this, OldCompiledSections);
 			ModifierGroupTasks.Add(ModifierGroupTask);
 			ModifierGroupTask->State = ETaskState::Ready;
 		}
 
-		ExecuteModifierGroupTasks(ModifierGroupTasks, World, MegaMesh, InPackageHelper);
+		ExecuteModifierGroupTasks(ModifierGroupTasks, World, MeshPartition, InPackageHelper);
 
 		int32 BuiltCompiledSections = 0;
 		int32 FailedCompiledSections = 0;
@@ -385,10 +409,10 @@ bool UWorldPartitionMeshPartitionBuilder::BuildMegaMesh(const AMeshPartition* Me
 
 		// destruct the tasks to free up the memory and actor references they are holding
 		ModifierGroupTasks.Empty();
-		MegaMeshEditorComponent->ResetGroupRegistry();
+		MeshPartitionEditorComponent->ResetGroupRegistry();
 
 		UE_LOG(LogMegaMeshEditor, Display, TEXT("-- %.6s -- %s Finish Variant (Built %d, Reused %d, Failed %d Compiled Sections)"),
-			   *MegaMeshGUID.ToString(), *BuildVariant.Name.ToString(),
+			   *MeshPartitionGUID.ToString(), *BuildVariant.Name.ToString(),
 			   BuiltCompiledSections, ReusedCompiledSections, FailedCompiledSections);
 
 		// Update stats
@@ -401,16 +425,16 @@ bool UWorldPartitionMeshPartitionBuilder::BuildMegaMesh(const AMeshPartition* Me
 		FWorldPartitionHelpers::DoCollectGarbage();
 	}
 
-	UE_LOG(LogMegaMeshEditor, Display, TEXT("-- %.6s -- Finished Mesh Partition Build (%s)"), *MegaMeshGUID.ToString(), *MegaMesh->GetActorNameOrLabel());
+	UE_LOG(LogMegaMeshEditor, Display, TEXT("-- %.6s -- Finished Mesh Partition Build (%s)"), *MeshPartitionGUID.ToString(), *MeshPartition->GetActorNameOrLabel());
 	return true;
 }
 
-bool UWorldPartitionMeshPartitionBuilder::CanBuildMegaMeshes(const UWorld* InWorld, FName InBuildOption)
+bool UWorldPartitionMeshPartitionBuilder::CanBuildMeshPartitions(const UWorld* InWorld, FName InBuildOption)
 {
 	return true;
 }
 
-EEditorBuildResult UWorldPartitionMeshPartitionBuilder::BuildMegaMeshes(UWorld* InWorld, EMeshPartitionReuseMode ReuseMode)
+EEditorBuildResult UWorldPartitionMeshPartitionBuilder::BuildMeshPartitions(UWorld* InWorld, EMeshPartitionReuseMode ReuseMode)
 {
 	check(InWorld);
 
@@ -425,9 +449,9 @@ EEditorBuildResult UWorldPartitionMeshPartitionBuilder::BuildMegaMeshes(UWorld*
 	}
 
 	// check for definitions
-	for (TActorIterator<AMeshPartition> MegaMeshIterator(InWorld); MegaMeshIterator; ++MegaMeshIterator)
+	for (TActorIterator<AMeshPartition> MeshPartitionIterator(InWorld); MeshPartitionIterator; ++MeshPartitionIterator)
 	{
-		const AMeshPartition* MeshPartition = *MegaMeshIterator;
+		const AMeshPartition* MeshPartition = *MeshPartitionIterator;
 		check(MeshPartition != nullptr);
 		
 		const UMeshPartitionDefinition* MeshPartitionDefinition = MeshPartition->GetMeshPartitionDefinition();
@@ -469,7 +493,7 @@ EEditorBuildResult UWorldPartitionMeshPartitionBuilder::BuildMegaMeshes(UWorld*
 		*ReuseModeString,
 		*LogFileArg,
 		*SCCProvider.GetName().ToString(),
-		CVarMegaMeshBuildCommandWaitForDebugger->GetBool() ? TEXT(" -waitforattach") : TEXT(""));
+		CVarMeshPartitionMeshBuildCommandWaitForDebugger->GetBool() ? TEXT(" -waitforattach") : TEXT(""));
 
 	FString FullLogPath = FPaths::ProjectLogDir() + LogFileName;
 	IPlatformFile& PlatformFile = FPlatformFileManager::Get().GetPlatformFile();
@@ -509,10 +533,13 @@ EEditorBuildResult UWorldPartitionMeshPartitionBuilder::BuildMegaMeshes(UWorld*
 				TArray<uint8> Buffer;
 
 				LogReader->Seek(LastReadPos);
-				Buffer.SetNumUninitialized(FileStat.FileSize - LastReadPos);
+				int64 BufferSizeBytes = FileStat.FileSize - LastReadPos;
+				constexpr int64 MaxBufferToProcess = 256*1024*1024;
+				BufferSizeBytes = (BufferSizeBytes > MaxBufferToProcess) ? MaxBufferToProcess : BufferSizeBytes;
+				Buffer.SetNumUninitialized(IntCastChecked<int32>(BufferSizeBytes));
 				LogReader->Serialize(Buffer.GetData(), Buffer.Num());
 
-				FString LogString = FString(StringCast<TCHAR>(reinterpret_cast<const char*>(Buffer.GetData())));
+				FString LogString = FString(StringCast<TCHAR>(reinterpret_cast<const char*>(Buffer.GetData()), Buffer.Num()));
 				FRegexMatcher Regex(LogProgressPattern, LogString);
 
 				while (Regex.FindNext())
@@ -542,15 +569,15 @@ void UWorldPartitionMeshPartitionBuilder::GatherDependenciesAndComputeHashes(FMo
 	using namespace EditorUtils;
 	FDependencyContext DependencyContext;
 
-	// gather dependencies from the megamesh definition
-	if (const UMeshPartitionDefinition* MegaMeshDefinition = InOutTask.BuildInfo.GetMegaMeshDefinition())
+	// gather dependencies from the MeshPartition definition
+	if (const UMeshPartitionDefinition* MeshPartitionDefinition = InOutTask.BuildInfo.GetMegaMeshDefinition())
 	{
-		DependencyContext.AddPackageDependency(MegaMeshDefinition);
-		DependencyContext.AddClassDependency(MegaMeshDefinition->GetClass());
-		MegaMeshDefinition->GatherDependencies(DependencyContext, InOutTask.BuildVariant.Name);
+		DependencyContext.AddPackageDependency(MeshPartitionDefinition);
+		DependencyContext.AddClassDependency(MeshPartitionDefinition->GetClass());
+		MeshPartitionDefinition->GatherDependencies(DependencyContext, InOutTask.BuildVariant.Name);
 	}
 
-	// gather dependencies from the megamesh modifiers
+	// gather dependencies from the MeshPartition modifiers
 	InOutTask.ModifierGroup.ForAllModifiers([&DependencyContext](MeshPartition::UModifierComponent* Modifier)
 	{
 		if (ensure(Modifier != nullptr))
@@ -577,7 +604,7 @@ void UWorldPartitionMeshPartitionBuilder::GatherDependenciesAndComputeHashes(FMo
 	InOutTask.BuildInfo.ClassHash = ComputeClassHash(Classes, &InOutTask.BuildInfo.ClassChecksums);
 }
 
-void UWorldPartitionMeshPartitionBuilder::ExecuteModifierGroupTasks(TIndirectArray<FModifierGroupTask>& ModifierGroupTasks, UWorld* InWorld, const AMeshPartition* InMegaMesh, FPackageSourceControlHelper& InPackageHelper)
+void UWorldPartitionMeshPartitionBuilder::ExecuteModifierGroupTasks(TIndirectArray<FModifierGroupTask>& ModifierGroupTasks, UWorld* InWorld, const AMeshPartition* InMeshPartition, FPackageSourceControlHelper& InPackageHelper)
 {
 	using namespace EditorUtils;
 
@@ -633,7 +660,7 @@ void UWorldPartitionMeshPartitionBuilder::ExecuteModifierGroupTasks(TIndirectArr
 						if (!Task.ReuseCompiledSectionActorDescGuid.IsValid() || (ReuseMode == EMeshPartitionReuseMode::ValidateHashes))
 						{
 							// load the modifiers, including all modifiers within level instances
-							LoadActors(InMegaMesh, Task);
+							LoadActors(InMeshPartition, Task);
 							Task.State = ETaskState::ActorLoaded;
 
 							// and compute an up-to-date modifier hash
@@ -675,7 +702,7 @@ void UWorldPartitionMeshPartitionBuilder::ExecuteModifierGroupTasks(TIndirectArr
 							SectionsReused++;
 							bUpdateLog = true;
 
-							WorldPartitionMegaMeshBuilderLocals::LogBuildProgress(GetTotalSectionNumber(), TotalNumberOfSectionsToBuild);
+							WorldPartitionMeshPartitionBuilderLocals::LogBuildProgress(GetTotalSectionNumber(), TotalNumberOfSectionsToBuild);
 						}
 						else
 						{
@@ -732,12 +759,12 @@ void UWorldPartitionMeshPartitionBuilder::ExecuteModifierGroupTasks(TIndirectArr
 					break;
 								
 				case ETaskState::ModifiersProcessed:
-					Task.CompiledSection = Task.MegaMeshEditorComponent->PrepareMegaMeshCompiledSection(Task.BuildInfo, Task.BuildVariant);
-					Task.MegaMeshEditorComponent->BuildMegaMeshCompiledSectionTextures(Task.CompiledSection, Task.ModifierGroup, Task.ProcessModifierTask.GetResult());
+					Task.CompiledSection = Task.MeshPartitionEditorComponent->PrepareMegaMeshCompiledSection(Task.BuildInfo, Task.BuildVariant);
+					Task.MeshPartitionEditorComponent->BuildMegaMeshCompiledSectionTextures(Task.CompiledSection, Task.ModifierGroup, Task.ProcessModifierTask.GetResult());
 
 					{
 						TArray<TWeakObjectPtr<MeshPartition::UModifierComponent>> ModifierPtrs = Task.ModifierGroup.AllResolvedModifierPtrs();
-						Task.MegaMeshEditorComponent->PostBuildSectionMesh(Task.CompiledSection, Task.ProcessModifierTask.GetResult(), ModifierPtrs);
+						Task.MeshPartitionEditorComponent->PostBuildSectionMesh(Task.CompiledSection, Task.ProcessModifierTask.GetResult(), ModifierPtrs);
 
 						MeshPartition::FTransformerUnit InitialUnit;
 
@@ -746,9 +773,9 @@ void UWorldPartitionMeshPartitionBuilder::ExecuteModifierGroupTasks(TIndirectArr
 						InitialUnit.bShouldRecomputeNormals = true;
 						InitialUnit.bShouldRecomputeTangents = true;
 
-						MeshPartition::UMeshPartitionDefinition* Definition = Task.MegaMeshEditorComponent->GetMegaMeshDefinition();
+						MeshPartition::UMeshPartitionDefinition* Definition = Task.MeshPartitionEditorComponent->GetMegaMeshDefinition();
 
-						Task.TransformerContext = Task.MegaMeshEditorComponent->LaunchTransformers(MoveTemp(InitialUnit), Definition, Task.BuildVariant);
+						Task.TransformerContext = Task.MeshPartitionEditorComponent->LaunchTransformers(MoveTemp(InitialUnit), Definition, Task.BuildVariant);
 					}
 
 					if (Task.TransformerContext.IsValid())
@@ -775,7 +802,7 @@ void UWorldPartitionMeshPartitionBuilder::ExecuteModifierGroupTasks(TIndirectArr
 
 					{
 						TArray<TWeakObjectPtr<MeshPartition::UModifierComponent>> ModifierPtrs = Task.ModifierGroup.AllResolvedModifierPtrs();
-						Task.MegaMeshEditorComponent->PostProcessSection(Task.CompiledSection, ModifierPtrs);
+						Task.MeshPartitionEditorComponent->PostProcessSection(Task.CompiledSection, ModifierPtrs);
 					}
 
 					AssignSectionToPlatformDataLayers(Task.CompiledSection);
@@ -789,7 +816,7 @@ void UWorldPartitionMeshPartitionBuilder::ExecuteModifierGroupTasks(TIndirectArr
 					NumTasksInFlight--;
 
 					int32 TotalSectionNumber = TotalBuiltCompiledSections + TotalFailedCompiledSections + TotalReusedCompiledSections + SectionsBuilt + SectionsReused;
-					WorldPartitionMegaMeshBuilderLocals::LogBuildProgress(GetTotalSectionNumber(), TotalNumberOfSectionsToBuild);
+					WorldPartitionMeshPartitionBuilderLocals::LogBuildProgress(GetTotalSectionNumber(), TotalNumberOfSectionsToBuild);
 
 					break;
 			}
@@ -811,17 +838,17 @@ void UWorldPartitionMeshPartitionBuilder::ExecuteModifierGroupTasks(TIndirectArr
 	}
 }
 
-void UWorldPartitionMeshPartitionBuilder::LoadActors(const AMeshPartition* InMegaMesh, FModifierGroupTask& InModifierGroupTask)
+void UWorldPartitionMeshPartitionBuilder::LoadActors(const AMeshPartition* InMeshPartition, FModifierGroupTask& InModifierGroupTask)
 {
 	TRACE_CPUPROFILER_EVENT_SCOPE(UWorldPartitionMeshPartitionBuilder::LoadActors);
 
-	if (InMegaMesh == nullptr)
+	if (InMeshPartition == nullptr)
 	{
 		return;
 	}
 
-	UMeshPartitionEditorComponent* MegaMeshEditorComponent = Cast<UMeshPartitionEditorComponent>(InMegaMesh->GetMeshPartitionComponent());
-	if (MegaMeshEditorComponent == nullptr)
+	UMeshPartitionEditorComponent* MeshPartitionEditorComponent = Cast<UMeshPartitionEditorComponent>(InMeshPartition->GetMeshPartitionComponent());
+	if (MeshPartitionEditorComponent == nullptr)
 	{
 		return;
 	}
@@ -832,12 +859,12 @@ void UWorldPartitionMeshPartitionBuilder::LoadActors(const AMeshPartition* InMeg
 		};
 
 
-	auto OnModifierLoaded = [this, MegaMeshEditorComponent, &InModifierGroupTask](MeshPartition::UModifierComponent* Modifier, bool bIsInLevelInstance)
+	auto OnModifierLoaded = [this, MeshPartitionEditorComponent, &InModifierGroupTask](MeshPartition::UModifierComponent* Modifier, bool bIsInLevelInstance)
 		{
-			MegaMeshEditorComponent->OnModifierLoaded(Modifier);
+			MeshPartitionEditorComponent->OnModifierLoaded(Modifier);
 		};
 
-	WorldPartitionHelpers::LoadAllRelevantMegaMeshActors(InModifierGroupTask.ModifierGroup, InMegaMesh, InModifierGroupTask.ActorRefs, OnActorLoaded, OnModifierLoaded);
+	WorldPartitionHelpers::LoadAllRelevantMegaMeshActors(InModifierGroupTask.ModifierGroup, InMeshPartition, InModifierGroupTask.ActorRefs, OnActorLoaded, OnModifierLoaded);
 }
 
 void UWorldPartitionMeshPartitionBuilder::AssignSectionToPlatformDataLayers(MeshPartition::ACompiledSection* InCompiledSection) const
diff --git a/Engine/Plugins/Experimental/MeshPartition/Source/MeshPartitionEditor/Public/WorldPartitionMeshPartitionBuilder.h b/Engine/Plugins/Experimental/MeshPartition/Source/MeshPartitionEditor/Public/WorldPartitionMeshPartitionBuilder.h
--- a/Engine/Plugins/Experimental/MeshPartition/Source/MeshPartitionEditor/Public/WorldPartitionMeshPartitionBuilder.h
+++ b/Engine/Plugins/Experimental/MeshPartition/Source/MeshPartitionEditor/Public/WorldPartitionMeshPartitionBuilder.h
@@ -40,7 +40,7 @@ ENUM_RANGE_BY_COUNT(UE::MeshPartition::EMeshPartitionReuseMode, UE::MeshPartitio
 
 namespace UE::MeshPartition
 {
-// Example Command Line: ProjectName MapName -run=WorldPartitionBuilderCommandlet -SCCProvider=Perforce -Builder=WorldPartitionMegaMeshBuilder
+// Example Command Line: ProjectName MapName -run=WorldPartitionBuilderCommandlet -SCCProvider=Perforce -Builder=WorldPartitionMeshPartitionBuilder
 UCLASS()
 class UWorldPartitionMeshPartitionBuilder : public UWorldPartitionBuilder
 {
@@ -86,8 +86,8 @@ class UWorldPartitionMeshPartitionBuilder : public UWorldPartitionBuilder
 		MeshPartition::FBuilderSettings 		BuilderSettings;
 		MeshPartition::FCompiledSectionBuildInfo	BuildInfo;
 		const MeshPartition::FCompiledSectionBuildVariant&		BuildVariant;
-		UMeshPartitionEditorComponent*			MegaMeshEditorComponent = nullptr;
-		UWorldPartitionMeshPartitionBuilder&		WorldPartitionMegaMeshBuilder;		
+		UMeshPartitionEditorComponent*			MeshPartitionEditorComponent = nullptr;
+		UWorldPartitionMeshPartitionBuilder&		WorldPartitionMeshPartitionBuilder;
 
 		// array of pre-existing compiled sections (potentially re-usable for incremental rebuild)
 		// NOTE: this is NOT thread safe, and only works because we always synchronously complete this task
@@ -110,23 +110,23 @@ class UWorldPartitionMeshPartitionBuilder : public UWorldPartitionBuilder
 		TSet<AActor*>						ModifierOwners;			// set of actors containing one of our modifier components
 		TSet<FWorldPartitionReference>		ActorRefs;				// references to world partition actors that must remain loaded while this task is running
 		
-		FModifierGroupTask(UE::MeshPartition::FModifierGroup&& InModifierGroup, const MeshPartition::FBuilderSettings& InBuilderSettings, const MeshPartition::FCompiledSectionBuildInfo& InBuildInfo, const MeshPartition::FCompiledSectionBuildVariant& InBuildVariant, UMeshPartitionEditorComponent* InMegaMeshEditorComponent, UWorldPartitionMeshPartitionBuilder& InWorldPartitionMegaMeshBuilder, TArrayView<const FCompiledSectionDescriptor> InOldCompiledSectionDescriptors)
+		FModifierGroupTask(UE::MeshPartition::FModifierGroup&& InModifierGroup, const MeshPartition::FBuilderSettings& InBuilderSettings, const MeshPartition::FCompiledSectionBuildInfo& InBuildInfo, const MeshPartition::FCompiledSectionBuildVariant& InBuildVariant, UMeshPartitionEditorComponent* InMeshPartitionEditorComponent, UWorldPartitionMeshPartitionBuilder& InWorldPartitionMeshPartitionBuilder, TArrayView<const FCompiledSectionDescriptor> InOldCompiledSectionDescriptors)
 			: ModifierGroup(MoveTemp(InModifierGroup))
 			, BuilderSettings(InBuilderSettings)
 			, BuildInfo(InBuildInfo)
 			, BuildVariant(InBuildVariant)
-			, MegaMeshEditorComponent(InMegaMeshEditorComponent)
-			, WorldPartitionMegaMeshBuilder(InWorldPartitionMegaMeshBuilder)
+			, MeshPartitionEditorComponent(InMeshPartitionEditorComponent)
+			, WorldPartitionMeshPartitionBuilder(InWorldPartitionMeshPartitionBuilder)
 			, OldCompiledSectionDescriptors(InOldCompiledSectionDescriptors)
 		{}
 
-		FModifierGroupTask(const UE::MeshPartition::FModifierGroup& InModifierGroup, const MeshPartition::FBuilderSettings& InBuilderSettings, const MeshPartition::FCompiledSectionBuildInfo& InBuildInfo, const MeshPartition::FCompiledSectionBuildVariant& InBuildVariant, UMeshPartitionEditorComponent* InMegaMeshEditorComponent, UWorldPartitionMeshPartitionBuilder& InWorldPartitionMegaMeshBuilder, TArrayView<const FCompiledSectionDescriptor> InOldCompiledSectionDescriptors)
+		FModifierGroupTask(const UE::MeshPartition::FModifierGroup& InModifierGroup, const MeshPartition::FBuilderSettings& InBuilderSettings, const MeshPartition::FCompiledSectionBuildInfo& InBuildInfo, const MeshPartition::FCompiledSectionBuildVariant& InBuildVariant, UMeshPartitionEditorComponent* InMeshPartitionEditorComponent, UWorldPartitionMeshPartitionBuilder& InWorldPartitionMeshPartitionBuilder, TArrayView<const FCompiledSectionDescriptor> InOldCompiledSectionDescriptors)
 			: ModifierGroup(InModifierGroup)
 			, BuilderSettings(InBuilderSettings)
 			, BuildInfo(InBuildInfo)
 			, BuildVariant(InBuildVariant)
-			, MegaMeshEditorComponent(InMegaMeshEditorComponent)
-			, WorldPartitionMegaMeshBuilder(InWorldPartitionMegaMeshBuilder)
+			, MeshPartitionEditorComponent(InMeshPartitionEditorComponent)
+			, WorldPartitionMeshPartitionBuilder(InWorldPartitionMeshPartitionBuilder)
 			, OldCompiledSectionDescriptors(InOldCompiledSectionDescriptors)
 		{}
 	};
@@ -140,30 +140,30 @@ public:
 	virtual bool CanProcessNonPartitionedWorlds() const override { return false; }
 	// UWorldPartitionBuilder interface end
 
-	static bool CanBuildMegaMeshes(const UWorld* InWorld, FName InBuildOption);
-	static EEditorBuildResult BuildMegaMeshes(UWorld* InWorld, EMeshPartitionReuseMode ReuseMode);
+	static bool CanBuildMeshPartitions(const UWorld* InWorld, FName InBuildOption);
+	static EEditorBuildResult BuildMeshPartitions(UWorld* InWorld, EMeshPartitionReuseMode ReuseMode);
 
 private:
 	/**
-	* Build a MegaMesh
+	* Build a MeshPartition
 	*/
-	bool BuildMegaMesh(const AMeshPartition* MegaMesh, const FGuid& BuildKey, UE::MeshPartition::FModifierDescriptorCache& DescriptorCache, FPackageSourceControlHelper& InPackageHelper);
+	bool BuildMeshPartition(const AMeshPartition* MeshPartition, const FGuid& BuildKey, UE::MeshPartition::FModifierDescriptorCache& DescriptorCache, FPackageSourceControlHelper& InPackageHelper);
 
 	/**
 	* Executes and manages all modifier group tasks upon completion.
+	* @Param InModifierGroupTasks The modifier group tasks to execute.
 	* @param InWorld The current world.
-	* @param InMegaMesh The AMegaMesh the modifiers belong to.
-	* @param InMegaMeshEditorComponent The UMeshPartitionEditorComponent performing the processing.
+	* @param InMeshPartition The AMeshPartition the modifiers belong to.
 	* @param InPackageHelper A helper for saving and source control operations.
 	*/
-	void ExecuteModifierGroupTasks(TIndirectArray<FModifierGroupTask>& ModifierGroupTasks, UWorld* InWorld, const AMeshPartition* InMegaMesh, FPackageSourceControlHelper& InPackageHelper);
+	void ExecuteModifierGroupTasks(TIndirectArray<FModifierGroupTask>& InModifierGroupTasks, UWorld* InWorld, const AMeshPartition* InMeshPartition, FPackageSourceControlHelper& InPackageHelper);
 
 	/**
 	* Loads the modifier owners. Populating the lookup table and refcounting the loaded modifier owners.
-	* @param InMegaMesh The AMegaMesh the modifiers belong to.
+	* @param InMeshPartition The AMeshPartition the modifiers belong to.
 	* @param InModifierGroupTask The ModifierGroupTask that needs processing.
 	*/
-	void LoadActors(const AMeshPartition* InMegaMesh, FModifierGroupTask& InModifierGroupTask);
+	void LoadActors(const AMeshPartition* InMeshPartition, FModifierGroupTask& InModifierGroupTask);
 
 	void AssignSectionToPlatformDataLayers(MeshPartition::ACompiledSection* InCompiledSection) const;
 
