Adding a mesh partition descriptor cache data structure to store the descriptors for an entire world

#Synced-CL 51539202

diff --git a/Engine/Plugins/Experimental/MeshPartition/Source/MeshPartitionEditor/Private/MeshPartitionDescriptorCache.cpp b/Engine/Plugins/Experimental/MeshPartition/Source/MeshPartitionEditor/Private/MeshPartitionDescriptorCache.cpp
--- a/Engine/Plugins/Experimental/MeshPartition/Source/MeshPartitionEditor/Private/MeshPartitionDescriptorCache.cpp
+++ b/Engine/Plugins/Experimental/MeshPartition/Source/MeshPartitionEditor/Private/MeshPartitionDescriptorCache.cpp
@@ -18,6 +18,18 @@
 namespace UE::MeshPartition
 {
 
+FCachedDescriptors::FCachedDescriptors(const FGuid& InMeshPartitionGUID, const AMeshPartition* InMeshPartitionActor)
+{
+	check(InMeshPartitionGUID.IsValid());
+	if (InMeshPartitionActor)
+	{
+		check(InMeshPartitionActor->GetActorGuid() == InMeshPartitionGUID);
+	}
+
+	this->MeshPartitionGUID = InMeshPartitionGUID;
+	this->MeshPartitionActor = InMeshPartitionActor;
+}
+
 FCachedDescriptors::FCachedDescriptors(const AMeshPartition* InMeshPartition, TConstArrayView<MeshPartition::FModifierDesc> InAllModifiersSorted, TArray<FCompiledSectionDescriptor>&& InCompiledSections)
 {
 	check(InMeshPartition != nullptr);
@@ -93,9 +105,7 @@ void FCachedDescriptors::ForAllBuildVariants(TFunctionRef<bool(const MeshPartiti
 	using FVariantNameCachedBuildVariantPair = TPair<FName, MeshPartition::FModifierGrouping>;
 
 	TArray<FName> Keys;
-
-    Keys.Reserve(BuildVariants.Num());
-
+	Keys.Reserve(BuildVariants.Num());
 	for (const FVariantNameCachedBuildVariantPair& VariantNameCachedBuildVariantPair : BuildVariants)
 	{
 		Keys.Emplace(VariantNameCachedBuildVariantPair.Key);
@@ -304,4 +314,208 @@ TSharedPtr<MeshPartition::FCachedDescriptors> FModifierDescriptorCache::GetCache
 		return GetCachedDescriptorsByGuid(MeshPartitionGuid);
 	}
 }
+
+TSharedPtr<MeshPartition::FCachedDescriptors> FWorldCachedDescriptors::GetCachedDescriptorsByPath(const FSoftObjectPath& MegaMeshPath) const
+{
+	if (const TSharedPtr<MeshPartition::FCachedDescriptors>* CacheEntryPtr = CachedDescriptorsByPath.Find(MegaMeshPath))
+	{
+		return *CacheEntryPtr;
+	}
+	return nullptr;
+}
+
+TSharedPtr<MeshPartition::FCachedDescriptors> FWorldCachedDescriptors::GetCachedDescriptorsByGuid(const FGuid& MeshPartitionGUID) const
+{
+	if (const TSharedPtr<MeshPartition::FCachedDescriptors>* CacheEntryPtr = CachedDescriptorsByGUID.Find(MeshPartitionGUID))
+	{
+		return *CacheEntryPtr;
+	}
+	return nullptr;
+}
+
+TSharedPtr<MeshPartition::FCachedDescriptors> FWorldCachedDescriptors::GetCachedDescriptors(const AMeshPartition& MeshPartitionActor) const
+{
+	// first try to look up by path
+	FSoftObjectPath MeshPartitionPath(&MeshPartitionActor);
+	if (const TSharedPtr<MeshPartition::FCachedDescriptors>* CacheEntryPtr = CachedDescriptorsByPath.Find(MeshPartitionPath))
+	{
+		return *CacheEntryPtr;
+	}
+	else
+	{
+		// otherwise look it up by guid
+		const FGuid& MeshPartitionGUID = MeshPartitionActor.GetActorGuid();
+		return GetCachedDescriptorsByGuid(MeshPartitionGUID);
+	}
+}
+
+TSharedPtr<MeshPartition::FCachedDescriptors> FWorldCachedDescriptors::FindOrAddByActor(const AMeshPartition& MeshPartitionActor)
+{
+	TSharedPtr<MeshPartition::FCachedDescriptors> Descriptors = GetCachedDescriptors(MeshPartitionActor);
+	if (Descriptors.IsValid())
+	{
+		check(Descriptors->MeshPartitionActor == &MeshPartitionActor);
+		check(Descriptors->MeshPartitionGUID == MeshPartitionActor.GetActorGuid());
+	}
+	else
+	{
+		// create a new entry
+		const FGuid& MeshPartitionGUID = MeshPartitionActor.GetActorGuid();
+		check(MeshPartitionGUID.IsValid());
+
+		FSoftObjectPath MeshPartitionPath(&MeshPartitionActor);
+		check(MeshPartitionPath.IsValid());
+
+		Descriptors = MakeShared<MeshPartition::FCachedDescriptors>(MeshPartitionGUID, &MeshPartitionActor);
+
+		// populate the lookups
+		AllMeshPartitionsCachedDescriptors.Add(Descriptors);
+		CachedDescriptorsByGUID.Add(MeshPartitionGUID, Descriptors);
+		CachedDescriptorsByPath.Add(MeshPartitionPath, Descriptors);
+	}
+	return Descriptors;
+}
+
+TSharedPtr<MeshPartition::FCachedDescriptors> FWorldCachedDescriptors::FindOrAddByGUID(const FGuid& MeshPartitionGUID)
+{
+	check(MeshPartitionGUID.IsValid());
+	TSharedPtr<MeshPartition::FCachedDescriptors> Descriptors = GetCachedDescriptorsByGuid(MeshPartitionGUID);
+	if (Descriptors.IsValid())
+	{
+		check(Descriptors->MeshPartitionGUID == MeshPartitionGUID);
+	}
+	else
+	{
+		// create a new entry (with no path)
+		Descriptors = MakeShared<MeshPartition::FCachedDescriptors>(MeshPartitionGUID, nullptr);
+
+		// populate the lookups (except path lookup)
+		AllMeshPartitionsCachedDescriptors.Add(Descriptors);
+		CachedDescriptorsByGUID.Add(MeshPartitionGUID, Descriptors);
+	}
+	return Descriptors;
+}
+
+TSharedPtr<FWorldCachedDescriptors> FModifierDescriptorCache::GetDescriptorsForAllMeshPartitionsInWorld(UWorld* InWorld)
+{
+	using namespace Utils;
+
+	check(InWorld);
+
+	TSharedPtr<FWorldCachedDescriptors> WorldDescriptors = MakeShared<FWorldCachedDescriptors>();
+
+	// check if this is a PIE world
+	bool bIsPIE = InWorld->IsPlayInEditor();
+
+	// first iterate the mesh partitions that exist, and add them to MegaMeshesInWorld
+	// (this does not include spatially streaming AMeshPartitions or those inside of level instances, but we don't support either of those cases)
+	for (TActorIterator<AMeshPartition> It(InWorld); It; ++It)
+	{
+		const AMeshPartition* MeshPartition = *It;
+		if (ensure(MeshPartition))
+		{
+			// populate an entry for this mesh partition actor
+			WorldDescriptors->FindOrAddByActor(*MeshPartition);
+		}
+	}
+
+	UWorldPartition* WorldPartition = InWorld->GetWorldPartition();
+	if (WorldPartition == nullptr)
+	{
+		UE_LOG(LogMegaMeshEditor, Warning, TEXT("World '%s' is not World Partitioned, and does not support MeshPartition operations"), *FSoftObjectPath(InWorld).ToString());
+		return WorldDescriptors;
+	}
+
+	UE::MeshPartition::EditorUtils::FPIEPathFixer& PIEPathFixer = UMeshPartitionEditorSubsystem::GetPIEPathFixer();
+
+	// Now iterate all of the descriptors in the world, and record them into the corresponding FCachedDescriptors
+	// NOTE: this only grabs actor descs from the Main world/level -- not from any level instances
+	ForEachMegaMeshDescInWorldPartition(WorldPartition,
+		[&WorldDescriptors, bIsPIE, &PIEPathFixer](MeshPartition::FModifierDesc& InDescriptor)
+		{
+			if (!InDescriptor.MegaMeshGuid.IsValid())
+			{
+				// TODO: we should warn about this modifier without an assigned mesh partition
+				return;
+			}
+
+			if (bIsPIE)
+			{
+					// replace the PIE package name with the non-PIE package name
+					PIEPathFixer.FixInPlace(InDescriptor.ModifierPath);
+			}
+
+			TSharedPtr<MeshPartition::FCachedDescriptors> Descriptors = WorldDescriptors->FindOrAddByGUID(InDescriptor.MegaMeshGuid);
+			// TODO: this is a slow add as it may insert in the middle, for O(N^2) worst case.  If we allowed modifier groups to be in an unsorted state, we could use a faster "AddUnsorted" instead
+			Descriptors->AllModifiersGroup.Add(InDescriptor);
+			WorldDescriptors->TotalModifiers++;
+		},
+		[&WorldDescriptors](FCompiledSectionDescriptor& Descriptor)
+		{
+			// ignore old descriptors that don't have megamesh GUID populated (these should not exist, outside of old test levels)
+			if (!Descriptor.Info.MegaMeshGUID.IsValid())
+			{
+				UE_LOG(LogMegaMeshEditor, Warning, TEXT("Encountered Compiled Section descriptor without a valid GUID"));
+				return;
+			}
+
+			TSharedPtr<MeshPartition::FCachedDescriptors> Descriptors = WorldDescriptors->FindOrAddByGUID(Descriptor.Info.MegaMeshGUID);
+			Descriptors->CompiledSections.Add(Descriptor);
+			WorldDescriptors->TotalCompiledSections++;
+		});
+
+	// Now iterate our collection of mesh partitions and Build the section groups for each
+	for (TSharedPtr<MeshPartition::FCachedDescriptors> Descriptors : WorldDescriptors->AllMeshPartitionsCachedDescriptors)
+	{
+		const FGuid& MeshPartitionGUID = Descriptors->MeshPartitionGUID;
+
+		const AMeshPartition* MeshPartitionActor = Descriptors->MeshPartitionActor.Get();
+		if (MeshPartitionActor == nullptr)
+		{
+			// no mesh partition actor -- possibly deleted, marked spatially streaming (unsupported), or moved to a level instance (unsupported)
+			// this is ok if we just have compiled sections, they could be leftover from before the mesh partition was deleted, and just need to be cleaned up.
+			// however if we still have modifiers referencing this mesh partition, then the user needs to remove/fix those
+			if (Descriptors->GetAllModifiers().Num() > 0)
+			{
+				UE_LOG(LogMegaMeshEditor, Warning, TEXT("Could not find mesh partition actor with GUID '%s', referenced by %d modifiers"), *MeshPartitionGUID.ToString(), Descriptors->GetAllModifiers().Num());
+			}
+		}
+		else
+		{
+			check(MeshPartitionActor->GetActorGuid() == MeshPartitionGUID);
+
+			const UMeshPartitionDefinition* MeshPartitionDefinition = MeshPartitionActor->GetMeshPartitionDefinition();
+			if (MeshPartitionDefinition == nullptr)
+			{
+				UE_LOG(LogCore, Warning, TEXT("No MegaMesh Definition provided for %s (%.6s), using default settings."), *MeshPartitionActor->GetName(), *MeshPartitionGUID.ToString());
+				MeshPartitionDefinition = UMeshPartitionDefinition::StaticClass()->GetDefaultObject<UMeshPartitionDefinition>();
+			}
+
+			// Sort the Modifiers by priority (and Base/Non-Base)
+			Descriptors->AllModifiersGroup.Sort(MeshPartitionDefinition->GetModifierTypePriorities());
+
+			if (Descriptors->AllModifiersGroup.BaseDescs().IsEmpty())
+			{
+				// no modifiers found.. this megamesh will be empty
+				FSoftObjectPath MeshPartitionActorPath(MeshPartitionActor);
+				UE_LOG(LogMegaMeshEditor, Warning, TEXT("Mesh Partition '%s' (GUID %s) does not have any base modifiers in the world, the mesh partition will be empty"), *MeshPartitionActorPath.ToString(), *MeshPartitionGUID.ToString());
+			}
+
+			// iterate all build variants and build the section grouping for each
+			TConstArrayView<MeshPartition::FCompiledSectionBuildVariant> BuildVariantArray = MeshPartitionDefinition->GetCompiledSectionBuildVariants();
+			check(BuildVariantArray.Num() > 0);
+			for (const MeshPartition::FCompiledSectionBuildVariant& BuildVariant : BuildVariantArray)
+			{
+				FModifierGrouping& NewGroup = Descriptors->BuildVariants.Add(BuildVariant.Name, MeshPartition::FModifierGrouping(BuildVariant, *MeshPartitionDefinition, Descriptors->AllModifiersGroup));
+				WorldDescriptors->TotalBuildVariants++;
+				WorldDescriptors->TotalSectionGroups += NewGroup.GetModifierGroups().Num();
+			}
+		}
+	}
+	
+	UE_LOG(LogCore, Verbose, TEXT("GetDescriptorsForAllMeshPartitionsInWorld (%d Modifiers, %d Compiled Sections, %d Variants, %d Groups)"), WorldDescriptors->TotalModifiers, WorldDescriptors->TotalCompiledSections, WorldDescriptors->TotalBuildVariants, WorldDescriptors->TotalSectionGroups);
+
+	return WorldDescriptors;
+}
+
 } // namespace UE::MeshPartition
diff --git a/Engine/Plugins/Experimental/MeshPartition/Source/MeshPartitionEditor/Public/MeshPartitionDescriptorCache.h b/Engine/Plugins/Experimental/MeshPartition/Source/MeshPartitionEditor/Public/MeshPartitionDescriptorCache.h
--- a/Engine/Plugins/Experimental/MeshPartition/Source/MeshPartitionEditor/Public/MeshPartitionDescriptorCache.h
+++ b/Engine/Plugins/Experimental/MeshPartition/Source/MeshPartitionEditor/Public/MeshPartitionDescriptorCache.h
@@ -45,10 +45,16 @@ public:
 struct FCachedDescriptors
 {
 	// NOTE: this struct is copy-on-write so it is thread-safe after construction (as long as you hold a TSharedPtr to it)
+	friend struct FWorldCachedDescriptors;
+	friend struct FModifierDescriptorCache;
 
 public:
 	FCachedDescriptors(const AMeshPartition* InMeshPartition, TConstArrayView<MeshPartition::FModifierDesc> InAllModifiersSorted, TArray<FCompiledSectionDescriptor>&& InCompiledSections);
 
+	// this constructor initializes the cached descriptors, but does not fill out the data.  Used in the construction by FWorldCachedDescriptors
+	// TODO: this should be private
+	FCachedDescriptors(const FGuid& InMeshPartitionGUID, const AMeshPartition* InMeshPartitionActor);
+
 private:
 	// actor GUID of the mesh partition actor
 	FGuid MeshPartitionGUID;
@@ -120,6 +126,38 @@ public:
 	void ForAllBuildVariants(TFunctionRef<bool(const MeshPartition::FModifierGrouping&)> InFunc) const;
 };
 
+// stores the FCachedDescriptors for all AMeshPartition actors in a world
+struct FWorldCachedDescriptors
+{
+	friend struct FModifierDescriptorCache;
+
+private:
+	TSoftObjectPtr<const UWorld> World;
+
+	int32 TotalCompiledSections = 0;
+	int32 TotalModifiers = 0;
+	int32 TotalBuildVariants = 0;
+	int32 TotalSectionGroups = 0;
+
+	TArray<TSharedPtr<MeshPartition::FCachedDescriptors>> AllMeshPartitionsCachedDescriptors;
+
+	// Mapping storing the cached descriptors for each mega mesh actor GUID.
+	TMap<FGuid, TSharedPtr<MeshPartition::FCachedDescriptors>> CachedDescriptorsByGUID;
+
+	// Mapping storing the cached descriptors for each mega mesh actor path.  We only cache one entry per unique path
+	// Bt there can be two paths pointing to the same FCachedDescriptors, if they both share a GUID
+	TMap<FSoftObjectPath, TSharedPtr<MeshPartition::FCachedDescriptors>> CachedDescriptorsByPath;
+
+	TSharedPtr<MeshPartition::FCachedDescriptors> FindOrAddByActor(const AMeshPartition& MeshPartitionActor);
+	TSharedPtr<MeshPartition::FCachedDescriptors> FindOrAddByGUID(const FGuid& MeshPartitionGUID);
+
+public:
+	TSharedPtr<MeshPartition::FCachedDescriptors> GetCachedDescriptors(const AMeshPartition& MeshPartitionActor) const;
+	TSharedPtr<MeshPartition::FCachedDescriptors> GetCachedDescriptorsByPath(const FSoftObjectPath& MeshPartitionPath) const;
+	TSharedPtr<MeshPartition::FCachedDescriptors> GetCachedDescriptorsByGuid(const FGuid& MeshPartitionGUID) const;
+	TConstArrayView<TSharedPtr<MeshPartition::FCachedDescriptors>> GetAllMeshPartitionsCachedDescriptors() const { return AllMeshPartitionsCachedDescriptors; }
+};
+
 struct FModifierDescriptorCache
 {
 public:
@@ -143,6 +181,8 @@ private:
 	// mapping storing the cached descriptors for each mega mesh actor
 	TMap<FSoftObjectPath, TSharedPtr<MeshPartition::FCachedDescriptors>> CachedDescriptors;
 
+	// returns the set of MeshPartition-relevant descriptors in the input World
+	TSharedPtr<FWorldCachedDescriptors> GetDescriptorsForAllMeshPartitionsInWorld(UWorld* World);
 };
 } // namespace UE::MeshPartition
 
