UE-365530

Mesh Terrain - Spline Modifier - Temporarily disabling processing of the Spline modifier in preview sections

#virtualized
#Synced-CL 51499417
#Synced-CL 51499417

diff --git a/Engine/Plugins/Experimental/MeshPartition/Source/MeshPartitionEditor/Private/Test/MeshPartitionSplineModifierTests.cpp b/Engine/Plugins/Experimental/MeshPartition/Source/MeshPartitionEditor/Private/Test/MeshPartitionSplineModifierTests.cpp
--- a/Engine/Plugins/Experimental/MeshPartition/Source/MeshPartitionEditor/Private/Test/MeshPartitionSplineModifierTests.cpp
+++ b/Engine/Plugins/Experimental/MeshPartition/Source/MeshPartitionEditor/Private/Test/MeshPartitionSplineModifierTests.cpp
@@ -606,6 +606,62 @@ TEST_CLASS_WITH_FLAGS(SplineModifierTests, "MeshPartition.Modifier.Splines", EAu
 				});
 	}
 
+	// QMetry: UE - TC - 20620	
+	TEST_METHOD(Verify_Temporarily_Disabling_Processing_Of_The_Spline_Modifier_In_Preview_Sections)
+	{
+		TestCommandBuilder
+			.Until(TEXT("Wait for initial modifier processing to complete"), [this]()
+				{
+					return !MeshPartitionEditorComponent->IsModifierParticipatingInActivePreviewSectionBuild(SplineModifier);
+				})
+			.Then(TEXT("Set the temporarily disable value to true"), [this]()
+				{
+					const FScopedTransaction Transaction(FText::FromString(TEXT("Temporarily disabled in Editor")));
+					SplineModifier->Modify();
+
+					SplineModifier->SetIsTemporarilyDisabledInEditor(true);
+					SplineModifier->UpdateSplineData();
+				})
+			.Until(TEXT("Wait for modifier update to process"), [this]()
+				{
+					return !MeshPartitionEditorComponent->IsModifierParticipatingInActivePreviewSectionBuild(SplineModifier);
+				})
+			.Then(TEXT("Verify that spline modifier has been temporarily disabled, and validate generated mesh against ground truth"), [this]()
+				{
+					ASSERT_THAT(IsTrue(SplineModifier->IsTemporarilyDisabledInEditor(), TEXT("Spline modifier is still enabled in editor")));
+					const FString TestName = TEXT("Verify_Temporarily_Disabling_Processing_Of_The_Spline_Modifier_In_Preview_Sections_Disabled");
+					CompareToReferenceAsset(TestName);
+				})
+			.Then(TEXT("Test Undo function on the feature"), [this]()
+				{
+					GEditor->UndoTransaction();
+				})
+			.Until(TEXT("Wait for modifier update to process"), [this]()
+				{
+					return !MeshPartitionEditorComponent->IsModifierParticipatingInActivePreviewSectionBuild(SplineModifier);
+				})
+			.Then(TEXT("Validate the effects of the Undo function: Temporarily disable value should be set to false again, and the modifier should be enable"), [this]()
+				{
+					ASSERT_THAT(IsFalse(SplineModifier->IsTemporarilyDisabledInEditor(), TEXT("Spline modifier is still disabled in editor")));
+					const FString TestName = TEXT("Verify_Temporarily_Disabling_Processing_Of_The_Spline_Modifier_In_Preview_Sections_Enabled");
+					CompareToReferenceAsset(TestName);
+				})
+			.Then(TEXT("Test Redo function on the feature"), [this]()
+				{
+					GEditor->RedoTransaction();
+				})
+			.Until(TEXT("Wait for modifier update to process"), [this]()
+				{
+					return !MeshPartitionEditorComponent->IsModifierParticipatingInActivePreviewSectionBuild(SplineModifier);
+				})
+			.Then(TEXT("Validate the effects of the Redo function: Temporarily disable value should be set to true, and the modifier should be disabled again"), [this]()
+				{
+					ASSERT_THAT(IsTrue(SplineModifier->IsTemporarilyDisabledInEditor(), TEXT("Spline modifier is still enabled in editor")));
+					const FString TestName = TEXT("Verify_Temporarily_Disabling_Processing_Of_The_Spline_Modifier_In_Preview_Sections_Disabled");
+					CompareToReferenceAsset(TestName);
+				});
+	}
+
 	AFTER_EACH()
 	{
 		ScopedEditorWorld.Reset();
diff --git a/Engine/Plugins/Experimental/MeshPartition/Source/MeshPartitionEditor/Public/MeshPartitionModifierComponent.h b/Engine/Plugins/Experimental/MeshPartition/Source/MeshPartitionEditor/Public/MeshPartitionModifierComponent.h
--- a/Engine/Plugins/Experimental/MeshPartition/Source/MeshPartitionEditor/Public/MeshPartitionModifierComponent.h
+++ b/Engine/Plugins/Experimental/MeshPartition/Source/MeshPartitionEditor/Public/MeshPartitionModifierComponent.h
@@ -576,6 +576,14 @@ public:
 	*/
 	virtual bool IsTemporarilyDisabledInEditor() const { return bIsTemporarilyDisabledInEditor; }
 
+	/**
+	* @param bInIsTemporarilyDisabledInEditor Sets whether to temporarily disable processing this modifier in any Mesh Partition builds.
+	*/
+	void SetIsTemporarilyDisabledInEditor(bool bInIsTemporarilyDisabledInEditor)
+	{
+		bIsTemporarilyDisabledInEditor = bInIsTemporarilyDisabledInEditor;
+	}
+
 	/**
 	* @return The raw value of the serialized bIsDisabled flag.
 	*/
@@ -713,7 +721,7 @@ private:
 	* Note: Builds containing temporarily disabled modifiers do not get written to DDC,
 	* as this state is expected to be reset on next map load. Also does not affect compiled section builds.
 	*/
-	UPROPERTY(EditInstanceOnly, Category = "Modifier|ModifierSettings", SkipSerialization, Transient, DuplicateTransient, TextExportTransient, Meta = (DisplayName = "Disabled only in Editor", EditCondition = "!bIsDisabled"))
+	UPROPERTY(EditInstanceOnly, Category = "Modifier|ModifierSettings", Transient, DuplicateTransient, TextExportTransient, Meta = (DisplayName = "Disabled only in Editor", EditCondition = "!bIsDisabled"))
 	bool bIsTemporarilyDisabledInEditor = false;
 
 	// The last bounds of the modifier that were integrated into the built mesh
