UE-365531

Mesh Terrain - Spline Modifier - Extending Spline modifier projection bounds in all directions

#Synced-CL 51499417
#Synced-CL 51499417

diff --git a/Engine/Plugins/Experimental/MeshPartition/Source/MeshPartitionEditor/Private/Modifiers/MeshPartitionSplineModifier.cpp b/Engine/Plugins/Experimental/MeshPartition/Source/MeshPartitionEditor/Private/Modifiers/MeshPartitionSplineModifier.cpp
--- a/Engine/Plugins/Experimental/MeshPartition/Source/MeshPartitionEditor/Private/Modifiers/MeshPartitionSplineModifier.cpp
+++ b/Engine/Plugins/Experimental/MeshPartition/Source/MeshPartitionEditor/Private/Modifiers/MeshPartitionSplineModifier.cpp
@@ -1036,6 +1036,7 @@ void USplineModifier::SetMeshedInteriorNumTriTarget(double InMeshedInteriorNumTr
 {
 	if (!FMath::IsFinite(InMeshedInteriorNumTriTarget))
 	{
+		UE_LOGFMT(LogUnrealMath, Warning, "USplineModifier::SetMeshedInteriorNumTriTarget: Invalid infinite input value: {Value}", InMeshedInteriorNumTriTarget);
 		return;
 	}
 
@@ -1046,6 +1047,7 @@ void USplineModifier::SetFalloffBoundsMultiplier(float InFalloffBoundsMultiplier
 {
 	if (!FMath::IsFinite(InFalloffBoundsMultiplier))
 	{
+		UE_LOGFMT(LogUnrealMath, Warning, "USplineModifier::InFalloffBoundsMultiplier: Invalid infinite input value: {Value}", InFalloffBoundsMultiplier);
 		return;
 	}
 
@@ -1056,4 +1058,20 @@ void USplineModifier::SetInteriorSmoothMode(ESplineModifierInteriorSmoothMode In
 	InteriorSmoothMode = InInteriorSmoothMode;
 }
 
+void USplineModifier::SetMaxProjectionHeightExtent(float InMaxProjectionHeightExtent)
+{
+	if (!FMath::IsFinite(InMaxProjectionHeightExtent))
+	{
+		UE_LOGFMT(LogUnrealMath, Warning, "USplineModifier::SetMaxProjectionHeightExtent: Invalid infinite input value: {Value}", InMaxProjectionHeightExtent);
+		return;
+	}
+
+	MaxProjectionHeightExtent = FMath::Max(1.0, InMaxProjectionHeightExtent);
+}
+
+void USplineModifier::SetMeshClosedInterior(bool InbMeshClosedInterior)
+{
+	bMeshClosedInterior = InbMeshClosedInterior;
+}
+
 } // namespace UE::MeshPartition
\ No newline at end of file
diff --git a/Engine/Plugins/Experimental/MeshPartition/Source/MeshPartitionEditor/Private/Test/MeshPartitionSplineModifierTests.cpp b/Engine/Plugins/Experimental/MeshPartition/Source/MeshPartitionEditor/Private/Test/MeshPartitionSplineModifierTests.cpp
--- a/Engine/Plugins/Experimental/MeshPartition/Source/MeshPartitionEditor/Private/Test/MeshPartitionSplineModifierTests.cpp
+++ b/Engine/Plugins/Experimental/MeshPartition/Source/MeshPartitionEditor/Private/Test/MeshPartitionSplineModifierTests.cpp
@@ -58,17 +58,26 @@ TEST_CLASS_WITH_FLAGS(SplineModifierTests, "MeshPartition.Modifier.Splines", EAu
 	const FName PriorityLayer = "TestType2";
 	const double SubPriority = 2.0;
 
-
 	/*
 	 * Specifies the falloff bounds multiplier for the relevant test.
 	 */
 	const float FalloffBounds = 10.0f;
 
 	/* 
-	 * Specifies the Z axis value to be set for projection bounds test.
+	 * Specifies the Z axis value to be set for projection bounds tests.
 	 */
 	const float ZAxisBoundsValue = 1000.0f;
 
+	/*
+	* Specifies the Max Projection Height Extent value to be set for the projection bound test.
+	*/
+	const float MaxProjectionHeightExtent = 2000.0f;
+
+	/* 
+	* Original bounds for the projection bound test.
+	*/
+	TArray<FBox> OriginalBounds;
+
 	BEFORE_EACH()
 	{
 		// Configure minimal editor world
@@ -193,6 +202,12 @@ TEST_CLASS_WITH_FLAGS(SplineModifierTests, "MeshPartition.Modifier.Splines", EAu
 		AdjustProjectionBoundsByZAxisTest(ZAxisBoundsValue);
 	}
 
+	// QMetry: UE - TC - 20616	
+	TEST_METHOD(Verify_Extending_Spline_Modifier_Projection_Bounds_In_All_Directions)
+	{
+		ExtendProjectionBoundsTest(ZAxisBoundsValue, MaxProjectionHeightExtent);
+	}
+
 	// QMetry: UE - TC - 20605	
 	TEST_METHOD(Verify_Using_Relative_Projection_Direction_For_Mesh_Projection_Translation_GroundTruth)
 	{
@@ -978,6 +993,99 @@ TEST_CLASS_WITH_FLAGS(SplineModifierTests, "MeshPartition.Modifier.Splines", EAu
 				});
 	}
 
+	void ExtendProjectionBoundsTest(float TargetZValue, float TargetMaxProjectionHeightExtent)
+	{
+		float OriginalMaxProjectionHeightExtent = SplineModifier->GetMaxProjectionHeightExtent();
+
+		TestCommandBuilder
+			.Until(TEXT("Wait for initial modifier processing to complete"), [this]()
+				{
+					return !MeshPartitionEditorComponent->IsModifierParticipatingInActivePreviewSectionBuild(SplineModifier);
+				})
+			.Then(TEXT("Set Location Z value to target value"), [this, TargetZValue]()
+				{
+					SplineModifier->SetWorldLocation(FVector(0.0f, 0.0f, TargetZValue));
+					SplineModifier->SetUseNearestSplineFrameForDisplacement(true);
+					SplineModifier->SetMeshClosedInterior(false);
+					SplineModifier->UpdateSplineData();
+				})
+			.Until(TEXT("Wait for modifier update to process"), [this]()
+				{
+					return !MeshPartitionEditorComponent->IsModifierParticipatingInActivePreviewSectionBuild(SplineModifier);
+				})
+			.Then(TEXT("Increase Max Projection Height Extent value"), [this, TargetMaxProjectionHeightExtent]()
+				{
+					OriginalBounds = SplineModifier->ComputeBounds();
+					ASSERT_THAT(IsTrue(OriginalBounds.Num() > 0, TEXT("OriginalBounds is empty — ComputeBounds returned no bounds")));
+
+					const FScopedTransaction Transaction(FText::FromString(TEXT("Set Max Projection Height Extent value to >=1000")));
+					SplineModifier->Modify();
+
+					SplineModifier->SetMaxProjectionHeightExtent(TargetMaxProjectionHeightExtent);
+					SplineModifier->UpdateSplineData();
+				})
+			.Until(TEXT("Wait for modifier to finish processing rebuild"), [this]()
+				{
+					return !MeshPartitionEditorComponent->IsModifierParticipatingInActivePreviewSectionBuild(SplineModifier);
+				})
+			.Then(TEXT("Verify that the spline modifier projection bounds have been increased in all axes"), [this]()
+				{
+					TArray<FBox> UpdatedBounds = SplineModifier->ComputeBounds();
+					ASSERT_THAT(IsTrue(UpdatedBounds.Num() > 0, TEXT("UpdatedBounds is empty — ComputeBounds returned no bounds")));
+					
+					ASSERT_THAT(IsTrue(UpdatedBounds[0].Min.X < OriginalBounds[0].Min.X, TEXT("The Min.X bound range did not expand")));
+					ASSERT_THAT(IsTrue(UpdatedBounds[0].Min.Y < OriginalBounds[0].Min.Y, TEXT("The Min.Y bound range did not expand")));
+					ASSERT_THAT(IsTrue(UpdatedBounds[0].Min.Z < OriginalBounds[0].Min.Z, TEXT("The Min.Z bound range did not expand")));
+					ASSERT_THAT(IsTrue(UpdatedBounds[0].Max.X > OriginalBounds[0].Max.X, TEXT("The Max.X bound range did not expand")));
+					ASSERT_THAT(IsTrue(UpdatedBounds[0].Max.Y > OriginalBounds[0].Max.Y, TEXT("The Max.Y bound range did not expand")));
+					ASSERT_THAT(IsTrue(UpdatedBounds[0].Max.Z > OriginalBounds[0].Max.Z, TEXT("The Max.Z bound range did not expand")));
+				})
+			.Then(TEXT("Test Undo function on the feature"), [this]()
+				{
+					GEditor->UndoTransaction();
+				})
+			.Until(TEXT("Wait for modifier update to process"), [this]()
+				{
+					return !MeshPartitionEditorComponent->IsModifierParticipatingInActivePreviewSectionBuild(SplineModifier);
+				})
+			.Then(TEXT("Validate the effects of the Undo function: The projection bounds have been reverted"), [this, OriginalMaxProjectionHeightExtent]()
+				{
+					ASSERT_THAT(IsTrue(FMath::IsNearlyEqual(SplineModifier->GetMaxProjectionHeightExtent(), OriginalMaxProjectionHeightExtent), TEXT("The Max Projection Height Extent value was not reverted to the original value")));
+
+					TArray<FBox> UpdatedBounds = SplineModifier->ComputeBounds();
+					ASSERT_THAT(IsTrue(UpdatedBounds.Num() > 0, TEXT("UpdatedBounds is empty — ComputeBounds returned no bounds")));
+
+					ASSERT_THAT(IsNear(UpdatedBounds[0].Min.X, OriginalBounds[0].Min.X, UE_DOUBLE_SMALL_NUMBER, TEXT("The Min.X bound range did not revert to the original value")));
+					ASSERT_THAT(IsNear(UpdatedBounds[0].Min.Y, OriginalBounds[0].Min.Y, UE_DOUBLE_SMALL_NUMBER, TEXT("The Min.Y bound range did not revert to the original value")));
+					ASSERT_THAT(IsNear(UpdatedBounds[0].Min.Z, OriginalBounds[0].Min.Z, UE_DOUBLE_SMALL_NUMBER, TEXT("The Min.Z bound range did not revert to the original value")));
+					ASSERT_THAT(IsNear(UpdatedBounds[0].Max.X, OriginalBounds[0].Max.X, UE_DOUBLE_SMALL_NUMBER, TEXT("The Max.X bound range did not revert to the original value")));
+					ASSERT_THAT(IsNear(UpdatedBounds[0].Max.Y, OriginalBounds[0].Max.Y, UE_DOUBLE_SMALL_NUMBER, TEXT("The Max.Y bound range did not revert to the original value")));
+					ASSERT_THAT(IsNear(UpdatedBounds[0].Max.Z, OriginalBounds[0].Max.Z, UE_DOUBLE_SMALL_NUMBER, TEXT("The Max.Z bound range did not revert to the original value")));
+				})
+			.Then(TEXT("Test Redo function on the feature"), [this]()
+				{
+					GEditor->RedoTransaction();
+				})
+			.Until(TEXT("Wait for modifier update to process"), [this]()
+				{
+					return !MeshPartitionEditorComponent->IsModifierParticipatingInActivePreviewSectionBuild(SplineModifier);
+				})
+			.Then(TEXT("Validate the effects of the Redo function: The projection bounds have been expanded again"), [this, TargetMaxProjectionHeightExtent]()
+				{
+					ASSERT_THAT(IsTrue(FMath::IsNearlyEqual(SplineModifier->GetMaxProjectionHeightExtent(), TargetMaxProjectionHeightExtent), TEXT("The Max Projection Height Extent value was not reverted to the expanded value")));
+
+					TArray<FBox> UpdatedBounds = SplineModifier->ComputeBounds();
+					ASSERT_THAT(IsTrue(UpdatedBounds.Num() > 0, TEXT("UpdatedBounds is empty — ComputeBounds returned no bounds")));
+
+					ASSERT_THAT(IsTrue(UpdatedBounds[0].Min.X < OriginalBounds[0].Min.X, TEXT("The Min.X bound range did not expand")));
+					ASSERT_THAT(IsTrue(UpdatedBounds[0].Min.Y < OriginalBounds[0].Min.Y, TEXT("The Min.Y bound range did not expand")));
+					ASSERT_THAT(IsTrue(UpdatedBounds[0].Min.Z < OriginalBounds[0].Min.Z, TEXT("The Min.Z bound range did not expand")));
+					ASSERT_THAT(IsTrue(UpdatedBounds[0].Max.X > OriginalBounds[0].Max.X, TEXT("The Max.X bound range did not expand")));
+					ASSERT_THAT(IsTrue(UpdatedBounds[0].Max.Y > OriginalBounds[0].Max.Y, TEXT("The Max.Y bound range did not expand")));
+					ASSERT_THAT(IsTrue(UpdatedBounds[0].Max.Z > OriginalBounds[0].Max.Z, TEXT("The Max.Z bound range did not expand")));
+				});
+	}
+
 	void RelativeProjectionDirectionTest(SplineTransformTestType Type)
 	{
 		TSharedPtr<FString> TransformTestName = MakeShared<FString>();
diff --git a/Engine/Plugins/Experimental/MeshPartition/Source/MeshPartitionEditor/Public/Modifiers/MeshPartitionSplineModifier.h b/Engine/Plugins/Experimental/MeshPartition/Source/MeshPartitionEditor/Public/Modifiers/MeshPartitionSplineModifier.h
--- a/Engine/Plugins/Experimental/MeshPartition/Source/MeshPartitionEditor/Public/Modifiers/MeshPartitionSplineModifier.h
+++ b/Engine/Plugins/Experimental/MeshPartition/Source/MeshPartitionEditor/Public/Modifiers/MeshPartitionSplineModifier.h
@@ -142,7 +142,13 @@ public:
 	}
 	UE_API void SetInteriorSmoothMode(ESplineModifierInteriorSmoothMode InInteriorSmoothMode);
 
+	float GetMaxProjectionHeightExtent() const
+	{
+		return MaxProjectionHeightExtent;
+	}
+	UE_API void SetMaxProjectionHeightExtent(float InMaxProjectionHeightExtent);
 
+	UE_API void SetMeshClosedInterior(bool InbMeshClosedInterior);
 
 private:
 
