Enable real-time channel texture updates during attribute painting

DynamicMesh can now be used to render the channel textures, fully supporting split attribute data required to handle this.
Update the mesh partition attribute paint tool to re-render channel weight textures as they are changed.

#Synced-CL 51758550

diff --git a/Engine/Plugins/Experimental/MeshPartition/Shaders/MeshPartitionMakeSectionChannels.usf b/Engine/Plugins/Experimental/MeshPartition/Shaders/MeshPartitionMakeSectionChannels.usf
--- a/Engine/Plugins/Experimental/MeshPartition/Shaders/MeshPartitionMakeSectionChannels.usf
+++ b/Engine/Plugins/Experimental/MeshPartition/Shaders/MeshPartitionMakeSectionChannels.usf
@@ -4,10 +4,11 @@
 
 // ----------------------------------------------------------------------------------
 
-Buffer<int>		InMeshIndices;	
+Buffer<int>		InMeshIndices;
 Buffer<float2>	InMeshUVs;
 Buffer<float>	InMeshChannels;
 Buffer<int>		InMeshOutlines;
+Buffer<int>		InUVElementToVertexID;
 
 uint4  InDrawCall; // NumVertices, NumChannelAttribs, FirstVertIndex
 uint   InChannelOffset; // Offset to the first channel element in the Channel buffer
@@ -21,24 +22,25 @@ void DrawUVDomainVS(
 	out float4  OutWeight : TEXCOORD1,
 	out float4 OutPosition : SV_POSITION)
 {
+	uint UVElementID = InMeshIndices[VertexIndex + InDrawCall.z];
+	uint VertexID = InUVElementToVertexID[UVElementID];
 
-	uint VID = InMeshIndices[VertexIndex + InDrawCall.z]; // Default vertex id
 	if (InDrawCall.w == 1) // OR outline sourcing path
 	{
 		uint EdgeIndex = VertexIndex >> 1;
 		uint EdgeVertexIndex = VertexIndex & 0x00000001;
-		VID = InMeshOutlines[EdgeIndex * 3 + EdgeVertexIndex];	
+		VertexID = InMeshOutlines[EdgeIndex * 3 + EdgeVertexIndex];	
 	}
 
-	SourceUV = InMeshUVs[VID];
-	OutWeight = InMeshChannels[InChannelOffset + VID];
+	SourceUV = InMeshUVs[UVElementID];
+	OutWeight = InMeshChannels[InChannelOffset + VertexID];
 
 	// Transform from UV space to ndc
-	float2 BaseUV = SourceUV;	
+	float2 BaseUV = SourceUV;
 	BaseUV *= 2.0;
 	OutPosition.x = BaseUV.x - 1.0;
 	OutPosition.y = 1.0 - BaseUV.y;
-    OutPosition.zw = float2(0.0, 1.0);
+	OutPosition.zw = float2(0.0, 1.0);
 }
 
 // ----------------------------------------------------------------------------------
diff --git a/Engine/Plugins/Experimental/MeshPartition/Source/MeshPartition/Public/MeshPartitionMeshData.h b/Engine/Plugins/Experimental/MeshPartition/Source/MeshPartition/Public/MeshPartitionMeshData.h
--- a/Engine/Plugins/Experimental/MeshPartition/Source/MeshPartition/Public/MeshPartitionMeshData.h
+++ b/Engine/Plugins/Experimental/MeshPartition/Source/MeshPartition/Public/MeshPartitionMeshData.h
@@ -529,6 +529,8 @@ struct FSectionChannelGenerationMeshData
 	TArray<FVector2f>		Texcoords;
 	TArray<float>			Channels;
 	TArray<FInt32Vector3>	Outlines;
+	TArray<int32>			UVElementToVertexID;
+	int32 					VertexCount;
 	double					DomainArea3D;
 	double					DomainAreaUV;
 	double					DomainGutterUV;
diff --git a/Engine/Plugins/Experimental/MeshPartition/Source/MeshPartitionCompute/Public/MeshPartitionChannelRasterizationShaders.h b/Engine/Plugins/Experimental/MeshPartition/Source/MeshPartitionCompute/Public/MeshPartitionChannelRasterizationShaders.h
--- a/Engine/Plugins/Experimental/MeshPartition/Source/MeshPartitionCompute/Public/MeshPartitionChannelRasterizationShaders.h
+++ b/Engine/Plugins/Experimental/MeshPartition/Source/MeshPartitionCompute/Public/MeshPartitionChannelRasterizationShaders.h
@@ -19,8 +19,9 @@ public:
 	BEGIN_SHADER_PARAMETER_STRUCT(FParameters, )
 		SHADER_PARAMETER_RDG_BUFFER_SRV(Buffer<uint32>, InMeshIndices)
 		SHADER_PARAMETER_RDG_BUFFER_SRV(Buffer<float2>, InMeshUVs)
-		SHADER_PARAMETER_RDG_BUFFER_SRV(Buffer<float>,  InMeshChannels)
+		SHADER_PARAMETER_RDG_BUFFER_SRV(Buffer<float>, InMeshChannels)
 		SHADER_PARAMETER_RDG_BUFFER_SRV(Buffer<uint32>, InMeshOutlines)
+		SHADER_PARAMETER_RDG_BUFFER_SRV(Buffer<uint32>, InUVElementToVertexID)
 		SHADER_PARAMETER(FUintVector4, InDrawCall)
 		SHADER_PARAMETER(uint32, InChannelOffset)
 	END_SHADER_PARAMETER_STRUCT()
diff --git a/Engine/Plugins/Experimental/MeshPartition/Source/MeshPartitionEditor/Private/MeshPartitionChannelCollection.cpp b/Engine/Plugins/Experimental/MeshPartition/Source/MeshPartitionEditor/Private/MeshPartitionChannelCollection.cpp
--- a/Engine/Plugins/Experimental/MeshPartition/Source/MeshPartitionEditor/Private/MeshPartitionChannelCollection.cpp
+++ b/Engine/Plugins/Experimental/MeshPartition/Source/MeshPartitionEditor/Private/MeshPartitionChannelCollection.cpp
@@ -102,31 +102,14 @@ UTexture* FChannelTextureRenderer::GetDefaultTexture()
 	return BlackTexture.LoadSynchronous();
 }
 
-Tasks::TTask<MeshPartition::FSectionChannels> FChannelTextureRenderer::BuildTextureForSection(const FMeshData& InSectionMesh, UObject* InOwner, bool bDownloadToAsset, const MeshPartition::FChannelMap& InChannels, float InTexelResolution)
+Tasks::TTask<TOptional<UTexture*>> FChannelTextureRenderer::LaunchAllocateTextureTask(
+	Tasks::TTask<FSection> InBuildSectionTask,
+	UObject* InOwner)
 {
-	TRACE_CPUPROFILER_EVENT_SCOPE(MeshPartition::FChannelTextureRenderer::BuildTextureForSection);
-	
-	Tasks::TTask<FSection> BuildSection = Tasks::Launch(TEXT("Generate UV Mesh"),
-		[&SectionMesh = InSectionMesh, ChannelMap = InChannels, InTexelResolution]()
+	return Tasks::Launch(TEXT("AllocateTexture"),
+		[InBuildSectionTask, InOwner]() mutable
 		{
-			FSection Section;
-
-			// Generate the UV domain for the section that unifies the bases in the section
-			// This populates the UVMesh in the Section struct, ready to be uploaded to gpu and rendered 
-			GenerateSectionMeshUVs(SectionMesh, Section, InTexelResolution);
-
-			// Generate the Channels as UVMesh attributes for the section
-			// Needs the channels declaration of the definition to make a stable mapping of the channel
-			GenerateSectionMeshChannels(SectionMesh, Section, ChannelMap);
-
-			return MoveTemp(Section);
-		}
-	);
-	
-	Tasks::TTask<TOptional<UTexture*>> AllocateTexture = Tasks::Launch(TEXT("AllocateTexture"),
-		[BuildSection, InOwner]() mutable
-		{
-			const FSection& Section = BuildSection.GetResult();
+			const FSection& Section = InBuildSectionTask.GetResult();
 
 			TOptional<UTexture*> SectionRenderTexture;
 
@@ -139,23 +122,54 @@ Tasks::TTask<MeshPartition::FSectionChannels> FChannelTextureRenderer::BuildText
 					Section.RasterResolution.Y,
 					NumActiveSlices,
 					PF_G8,
-					Outer, // new created texture is parented to owner
+					Outer,
 					MakeUniqueObjectName(Outer, UTexture::StaticClass(), TEXT("MeshPartitionChannelTexture"))
 				);
 			}
 
 			return MoveTemp(SectionRenderTexture);
 		},
-		Tasks::Prerequisites(BuildSection)
+		Tasks::Prerequisites(InBuildSectionTask)
 	);
+}
 
-	Tasks::TTask<MeshPartition::FSectionChannels> RenderChannels = Tasks::Launch(TEXT("RenderSectionChannels"),
-		[BuildSection, AllocateTexture, bDownloadToAsset]() mutable
+Tasks::TTask<MeshPartition::FSectionChannels> FChannelTextureRenderer::LaunchRenderTextureTask(
+	Tasks::TTask<FSection> InBuildSectionTask,
+	Tasks::TTask<TOptional<UTexture*>> InAllocateTextureTask,
+	bool bInDownloadToAsset,
+	bool bInExtractSectionMeshData
+) {
+	return Tasks::Launch(TEXT("RenderSectionChannels"),
+		[InBuildSectionTask, InAllocateTextureTask, bInDownloadToAsset, bInExtractSectionMeshData]() mutable
 		{
-			FSection Section = MoveTemp(BuildSection.GetResult());
-			const TOptional<UTexture*>& SectionRenderTexture = AllocateTexture.GetResult();
+			FSection Section = MoveTemp(InBuildSectionTask.GetResult());
+			TArray<uint8> ChannelOffsetTable = Section.ChannelOffsetTable;
+			const FVector2f TexcoordMetrics = Section.TexcoordMetrics;
+			const int32 NumActiveSlices = Section.NumActivePages();
 
-			int32 NumActiveSlices = Section.NumActivePages();
+			MeshPartition::FSectionChannels Result {
+				.Table = MoveTemp(ChannelOffsetTable),
+				.TexcoordMetrics = TexcoordMetrics,
+			};
+
+			if (bInExtractSectionMeshData)
+			{
+				TRACE_CPUPROFILER_EVENT_SCOPE(FChannelTextureRenderer::ExtractSectionMeshData);
+
+				Result.Mesh = {
+					Section.UVmesh.Indices,
+					Section.UVmesh.UVs,
+					Section.UVmesh.Channels,
+					Section.UVmesh.Outlines,
+					Section.UVmesh.UVElementToVertexID,
+					Section.UVmesh.VertexCount,
+					Section.DomainMapping.Area3D,
+					Section.DomainMapping.AreaUV,
+					Section.DomainMapping.GetGutterUV()
+				};
+			}
+
+			const TOptional<UTexture*>& SectionRenderTexture = InAllocateTextureTask.GetResult();
 
 			TStrongObjectPtr<UTexture> ChannelTexture(nullptr);
 
@@ -170,10 +184,10 @@ Tasks::TTask<MeshPartition::FSectionChannels> FChannelTextureRenderer::BuildText
 				SectionRenderTexture.GetValue()->UpdateResource();
 
 				// Render section channels
-				RenderSectionChannels(SectionRenderTexture.GetValue(), Section);
+				RenderSectionChannels(SectionRenderTexture.GetValue(), MoveTemp(Section));
 
 				// and eventually download the texels to the asset in the case of the compiled section
-				if (bDownloadToAsset)
+				if (bInDownloadToAsset)
 				{
 					DownloadSectionTextureToAsset(SectionRenderTexture.GetValue());
 				}
@@ -192,29 +206,72 @@ Tasks::TTask<MeshPartition::FSectionChannels> FChannelTextureRenderer::BuildText
 				);
 			}
 
-			MeshPartition::FSectionChannels Result {
-				.Table = MoveTemp(Section.ChannelOffsetTable),
-				.TexcoordMetrics = Section.TexcoordMetrics,
-				.Texture = ChannelTexture,
-				.Mesh = { MoveTemp(Section.UVmesh.Indices), MoveTemp(Section.UVmesh.UVs), MoveTemp(Section.UVmesh.Channels), MoveTemp(Section.UVmesh.Outlines),
-							Section.DomainMapping.Area3D, Section.DomainMapping.AreaUV, Section.DomainMapping.GetGutterUV() }
-			};
-
+			Result.Texture = ChannelTexture;
 			return Result;
 		},
-		Tasks::Prerequisites(AllocateTexture),
+		Tasks::Prerequisites(InAllocateTextureTask),
 		Tasks::ETaskPriority::Normal,
 		Tasks::EExtendedTaskPriority::GameThreadNormalPri
 	);
+}
+
+Tasks::TTask<MeshPartition::FSectionChannels> FChannelTextureRenderer::BuildTextureForSection(const FMeshData& InSectionMesh, UObject* InOwner, bool bDownloadToAsset, const MeshPartition::FChannelMap& InChannels, float InTexelResolution)
+{
+	TRACE_CPUPROFILER_EVENT_SCOPE(MeshPartition::FChannelTextureRenderer::BuildTextureForSection);
+
+	Tasks::TTask<FSection> BuildSection = Tasks::Launch(TEXT("Generate UV Mesh"),
+		[&SectionMesh = InSectionMesh, ChannelMap = InChannels, InTexelResolution]()
+		{
+			FSection Section;
+
+			// Generate the UV domain for the section that unifies the bases in the section
+			// This populates the UVMesh in the Section struct, ready to be uploaded to gpu and rendered 
+			GenerateSectionMeshUVs(SectionMesh, Section, InTexelResolution);
+
+			// Generate the Channels as UVMesh attributes for the section
+			// Needs the channels declaration of the definition to make a stable mapping of the channel
+			GenerateSectionMeshChannels(SectionMesh, Section, ChannelMap);
+
+			return MoveTemp(Section);
+		}
+	);
 
-	return RenderChannels;
+	Tasks::TTask<TOptional<UTexture*>> AllocateTexture = LaunchAllocateTextureTask(BuildSection, InOwner);
+	return LaunchRenderTextureTask(BuildSection, AllocateTexture, bDownloadToAsset);
+}
+
+Tasks::TTask<MeshPartition::FSectionChannels> FChannelTextureRenderer::BuildTextureForSectionWithCachedTopology(
+	const FDynamicMesh3& InSectionMesh,
+	const FChannelRenderUVMeshTopology& InCachedTopology,
+	UObject* InOwner,
+	bool bDownloadToAsset,
+	const MeshPartition::FChannelMap& InChannels
+)
+{
+	TRACE_CPUPROFILER_EVENT_SCOPE(MeshPartition::FChannelTextureRenderer::BuildTextureForSectionWithCachedTopology);
+
+	Tasks::TTask<FSection> BuildSection = Tasks::Launch(TEXT("BuildSectionFromCache"),
+		[&SectionMesh = InSectionMesh, &CachedTopology = InCachedTopology, ChannelMap = InChannels]() mutable
+		{
+			// Apply cached topology
+			FSection Section = BuildSectionFromCachedTopology(CachedTopology);
+
+			// Generate just the channels data from current mesh data
+			GenerateSectionMeshChannels(SectionMesh, Section, ChannelMap);
+
+			return MoveTemp(Section);
+		}
+	);
+
+	Tasks::TTask<TOptional<UTexture*>> AllocateTexture = LaunchAllocateTextureTask(BuildSection, InOwner);
+	return LaunchRenderTextureTask(BuildSection, AllocateTexture, bDownloadToAsset);
 }
 
 void FChannelTextureRenderer::GenerateSectionMeshChannels(const FMeshData& InSectionMesh, FChannelTextureRenderer::FSection& InSection, const MeshPartition::FChannelMap& InChannels)
 {
 	TRACE_CPUPROFILER_EVENT_SCOPE(MeshPartition::FChannelTextureRenderer::GenerateSectionMeshChannels);
 
-	int32 NumVertices = InSection.UVmesh.NumVertices();
+	int32 NumVertices = InSection.UVmesh.VertexCount;
 	int32 MaxNumChannels = InChannels.GetNumChannels();
 
 	TArray<int32> ChannelActiveCounts;
@@ -259,7 +316,122 @@ void FChannelTextureRenderer::GenerateSectionMeshChannels(const FMeshData& InSec
 
 		// gather the section channels values
 		// and allocate new channel entry if needed
-		InSection.UVmesh.Channels.Init(0, NumVertices * SourceSectionChannels.Num());
+		InSection.UVmesh.Channels.SetNum(NumVertices * SourceSectionChannels.Num());
+
+		uint64 ChannelWeightsOffset = NumVertices * sizeof(float);
+		for (int32 i = 0; i < SourceSectionChannels.Num(); ++i)
+		{
+			FMemory::Memcpy((InSection.UVmesh.Channels.GetData() + i * NumVertices), SourceSectionChannels[i].GetData(), NumVertices * sizeof(float));
+		}
+	}
+	InSection.UVmesh.ChannelActiveCounts = ChannelActiveCounts;
+
+	// Gather the actual active channels with non zero values
+	// Populate the channel offset table
+	int8 ActiveChannelIndex = 0;
+	TArray<uint8> ActiveChannelIndices;
+	TArray<uint8> ChannelTableIndices;
+	for (int32 C : ChannelActiveCounts)
+	{
+		if (C > 0)
+		{
+			// This channel is active, and will be stored there
+			ChannelTableIndices.Add(ActiveChannelIndices.Num());
+			ActiveChannelIndices.Add(ActiveChannelIndex);
+		}
+		else
+		{
+			ChannelTableIndices.Add(-1); // empty channel
+		}
+		ActiveChannelIndex++;
+	}
+	InSection.UVmesh.ActiveChannelIndices = ActiveChannelIndices;
+	InSection.ChannelOffsetTable = ChannelTableIndices;
+}
+
+void FChannelTextureRenderer::GenerateSectionMeshChannels(const Geometry::FDynamicMesh3& InSectionMesh, FChannelTextureRenderer::FSection& InSection, const MeshPartition::FChannelMap& InChannels)
+{
+	TRACE_CPUPROFILER_EVENT_SCOPE(MeshPartition::FChannelTextureRenderer::GenerateSectionMeshChannels);
+
+	const int32 NumVertices = InSectionMesh.VertexCount();
+	const int32 MaxNumChannels = InChannels.GetNumChannels();
+
+	TArray<int32> ChannelActiveCounts;
+	ChannelActiveCounts.SetNumZeroed(MaxNumChannels);
+
+	const Geometry::FDynamicMeshAttributeSet* Attributes = InSectionMesh.Attributes();
+
+	if (Attributes->NumWeightLayers() > 0)
+	{
+		// Build dense weight array matching vertex ordering
+		// Create mapping from sparse VertexID to dense index
+		TArray<int32> VertexIDToDenseIndex;
+		VertexIDToDenseIndex.SetNum(InSectionMesh.MaxVertexID());
+		{
+			int32 DenseIndex = 0;
+			for (int VertexID_Map : InSectionMesh.VertexIndicesItr())
+			{
+				VertexIDToDenseIndex[VertexID_Map] = DenseIndex++;
+			}
+			ensure(DenseIndex == NumVertices);
+		}
+
+		// Grab the local packed array of channel attributes from the weight layer values
+		const TArray<FName> MeshPartitionChannels = InChannels.GetChannels();
+		TArray<int32> SectionChannelIndex;
+		TArray<TArray<float>> SourceSectionChannels;
+		int NumUsedChannels = 0;
+		int32 ChannelIndex = 0;
+		for (const FName ChannelName : MeshPartitionChannels)
+		{
+			const Geometry::FDynamicMeshWeightAttribute* WeightLayerAttribute = nullptr;
+			// Find the WeightLayerIndex on the dynamic mesh matching the current channel name, if it exists.
+			for (int WeightLayerIndex = 0; WeightLayerIndex < Attributes->NumWeightLayers(); ++WeightLayerIndex)
+			{
+				auto WeightLayer = Attributes->GetWeightLayer(WeightLayerIndex);
+				const FName WeightLayerName = WeightLayer->GetName();
+				if (WeightLayerName == ChannelName)
+				{
+					WeightLayerAttribute = WeightLayer;
+					break;
+				}
+			}
+
+			if (WeightLayerAttribute)
+			{
+				TArray<float> WeightLayer;
+				WeightLayer.SetNumZeroed(NumVertices);
+
+				int32 ActiveCount = 0;
+				for (int VertexID : InSectionMesh.VertexIndicesItr())
+				{
+					float Value = 0.0f;
+					WeightLayerAttribute->GetValue(VertexID, &Value);
+
+					int32 DenseIdx = VertexIDToDenseIndex[VertexID];
+					WeightLayer[DenseIdx] = Value;
+
+					if (Value > 0.0f)
+					{
+						++ActiveCount;
+					}
+				}
+
+				// Channel is active if it has non-zero values
+				if (ActiveCount > 0)
+				{
+					SourceSectionChannels.Emplace(MoveTemp(WeightLayer));
+					SectionChannelIndex.Add(ChannelIndex);
+					ChannelActiveCounts[ChannelIndex] = ActiveCount;
+				}
+			}
+
+			ChannelIndex++;
+		}
+
+		// gather the section channels values
+		// and allocate new channel entry if needed
+		InSection.UVmesh.Channels.SetNum(NumVertices * SourceSectionChannels.Num());
 
 		uint64 ChannelWeightsOffset = NumVertices * sizeof(float);
 		for (int32 i = 0; i < SourceSectionChannels.Num(); ++i)
@@ -347,18 +519,21 @@ UTexture* FChannelTextureRenderer::AllocateSectionTextureResource(int32 InSizeX,
 	return SectionTexture;
 }
 
-void FChannelTextureRenderer::RenderSectionChannels(UTexture* InOutSectionTexture, const FChannelTextureRenderer::FSection& InSection) 
+void FChannelTextureRenderer::RenderSectionChannels(UTexture* InOutSectionTexture, FChannelTextureRenderer::FSection InSection) 
 {
 	TRACE_CPUPROFILER_EVENT_SCOPE(MeshPartition::FChannelTextureRenderer::RenderSectionChannels);
 
-	ENQUEUE_RENDER_COMMAND(MegaMeshRenderRenderProceduralPainting)([DestSectionTexture = InOutSectionTexture, Section = InSection](FRHICommandListImmediate& RHICmdList)
+	ENQUEUE_RENDER_COMMAND(MegaMeshRenderRenderProceduralPainting)([DestSectionTexture = InOutSectionTexture, Section = MoveTemp(InSection)](FRHICommandListImmediate& RHICmdList)
 		{
 			TRACE_CPUPROFILER_EVENT_SCOPE(MeshPartition::FChannelTextureRenderer::Render_Command_MegaMeshRenderRenderProceduralPainting);
 			
 			int32 NumIndices = Section.UVmesh.NumIndices();
-			int32 NumVertices = Section.UVmesh.NumVertices();
+			int32 NumUVs = Section.UVmesh.UVs.Num();
+			int32 ActualVertexCount = Section.UVmesh.VertexCount;
 			int32 NumChannels = Section.UVmesh.NumChannels();
 			int32 NumOutlines = Section.UVmesh.Outlines.Num();
+			
+			const bool bShouldRenderOutlines = NumOutlines > 0;
 
 			TArray<uint8> ChannelOffsetTable = Section.ChannelOffsetTable;
 
@@ -366,6 +541,8 @@ void FChannelTextureRenderer::RenderSectionChannels(UTexture* InOutSectionTextur
 
 			MeshPartition::FRenderCaptureManager::BeginCapture(GraphBuilder);
 
+			TRACE_CPUPROFILER_EVENT_MANUAL_START("UploadBuffers");
+
 			// Section Indices
 			FRDGBufferSRVRef RDGSectionMeshIndicesSRV;
 			{
@@ -384,11 +561,11 @@ void FChannelTextureRenderer::RenderSectionChannels(UTexture* InOutSectionTextur
 			FRDGBufferSRVRef RDGSectionMeshUVsSRV;
 			{
 				// Section turned into RDG Resource
-				FRDGBufferRef RDGSectionBuffer = GraphBuilder.CreateBuffer(FRDGBufferDesc::CreateBufferDesc(2 * sizeof(float), NumVertices), TEXT("SectionUVs"));
-				GraphBuilder.QueueBufferUpload(RDGSectionBuffer, Section.UVmesh.UVs.GetData(), 2 * sizeof(float) * NumVertices);
+				FRDGBufferRef RDGSectionBuffer = GraphBuilder.CreateBuffer(FRDGBufferDesc::CreateBufferDesc(2 * sizeof(float), NumUVs), TEXT("SectionUVs"));
+				GraphBuilder.QueueBufferUpload(RDGSectionBuffer, Section.UVmesh.UVs.GetData(), 2 * sizeof(float) * NumUVs);
 
 				FRDGBufferSRVDesc RDGSectionBufferDesc(RDGSectionBuffer, PF_G32R32F);
-				RDGSectionBufferDesc.NumElements = NumVertices;
+				RDGSectionBufferDesc.NumElements = NumUVs;
 				RDGSectionBufferDesc.StartOffsetBytes = 0;
 
 				RDGSectionMeshUVsSRV = GraphBuilder.CreateSRV(RDGSectionBufferDesc);
@@ -398,17 +575,30 @@ void FChannelTextureRenderer::RenderSectionChannels(UTexture* InOutSectionTextur
 			FRDGBufferSRVRef RDGSectionMeshChannelsSRV;
 			{
 				// Section turned into RDG Resource
-				FRDGBufferRef RDGSectionBuffer = GraphBuilder.CreateBuffer(FRDGBufferDesc::CreateBufferDesc(sizeof(float), NumChannels * NumVertices), TEXT("SectionChannelMasks"));
-				GraphBuilder.QueueBufferUpload(RDGSectionBuffer, Section.UVmesh.Channels.GetData(), sizeof(float) * NumChannels * NumVertices);
+				FRDGBufferRef RDGSectionBuffer = GraphBuilder.CreateBuffer(FRDGBufferDesc::CreateBufferDesc(sizeof(float), Section.UVmesh.Channels.Num()), TEXT("SectionChannelMasks"));
+				GraphBuilder.QueueBufferUpload(RDGSectionBuffer, Section.UVmesh.Channels.GetData(), sizeof(float) * Section.UVmesh.Channels.Num());
 
 				FRDGBufferSRVDesc RDGSectionBufferDesc(RDGSectionBuffer, PF_R32_FLOAT);
-				RDGSectionBufferDesc.NumElements = NumChannels * NumVertices;
+				RDGSectionBufferDesc.NumElements = Section.UVmesh.Channels.Num();
 				RDGSectionBufferDesc.StartOffsetBytes = 0;
 
 				RDGSectionMeshChannelsSRV = GraphBuilder.CreateSRV(RDGSectionBufferDesc);
 			}
 
-			// Section Channels
+			// UV Element to Vertex ID Indirection Buffer
+			FRDGBufferSRVRef RDGUVElementToVertexIDSRV;
+			{
+				FRDGBufferRef RDGIndirectionBuffer = GraphBuilder.CreateBuffer(FRDGBufferDesc::CreateBufferDesc(sizeof(int32), NumUVs), TEXT("UVElementToVertexIDIndirection"));
+				GraphBuilder.QueueBufferUpload(RDGIndirectionBuffer, Section.UVmesh.UVElementToVertexID.GetData(), sizeof(int32) * NumUVs);
+
+				FRDGBufferSRVDesc RDGIndirectionBufferDesc(RDGIndirectionBuffer, PF_R32_SINT);
+				RDGIndirectionBufferDesc.NumElements = NumUVs;
+				RDGIndirectionBufferDesc.StartOffsetBytes = 0;
+
+				RDGUVElementToVertexIDSRV = GraphBuilder.CreateSRV(RDGIndirectionBufferDesc);
+			}
+
+			// Section Outlines
 			FRDGBufferSRVRef RDGSectionMeshOutlinesSRV;
 			{
 				// Section turned into RDG Resource
@@ -421,6 +611,7 @@ void FChannelTextureRenderer::RenderSectionChannels(UTexture* InOutSectionTextur
 
 				RDGSectionMeshOutlinesSRV = GraphBuilder.CreateSRV(RDGSectionBufferDesc);
 			}
+			TRACE_CPUPROFILER_EVENT_MANUAL_END();
 
 			// Pass Shaders
 			FGlobalShaderMap* ShaderMap = GetGlobalShaderMap(GMaxRHIFeatureLevel);
@@ -453,13 +644,17 @@ void FChannelTextureRenderer::RenderSectionChannels(UTexture* InOutSectionTextur
 
 			for (int32 Page = 0; Page < Section.NumActivePages(); ++Page)
 			{
+				// Get definition channel index for this Page
+				const int32 DefinitionChannelIndex = Section.UVmesh.ActiveChannelIndices[Page];
+
 				// Pass Parameters
 				FMeshPartition_DrawUVDomain_Parameters* PassParameters = GraphBuilder.AllocParameters<FMeshPartition_DrawUVDomain_Parameters>();
 				PassParameters->VS.InMeshIndices = RDGSectionMeshIndicesSRV;
 				PassParameters->VS.InMeshUVs = RDGSectionMeshUVsSRV;
 				PassParameters->VS.InMeshChannels = RDGSectionMeshChannelsSRV;
 				PassParameters->VS.InMeshOutlines = RDGSectionMeshOutlinesSRV;
-				PassParameters->VS.InChannelOffset = Page * NumVertices; // Channel begining offset in the Channels buffer
+				PassParameters->VS.InUVElementToVertexID = RDGUVElementToVertexIDSRV;
+				PassParameters->VS.InChannelOffset = Page * ActualVertexCount;
 
 				ClearUnusedGraphResources(VertexShader, &PassParameters->VS);
 				ClearUnusedGraphResources(PixelShader, &PassParameters->PS);
@@ -472,7 +667,7 @@ void FChannelTextureRenderer::RenderSectionChannels(UTexture* InOutSectionTextur
 					RDG_EVENT_NAME("MeshPartition_MakeSectionChannels"),
 					PassParameters,
 					ERDGPassFlags::Raster,
-					[PassParameters, Page, Section, ShaderMap, VertexShader, PixelShader](FRDGAsyncTask, FRHICommandList& RHICmdList)
+					[PassParameters, Page, &Section, ShaderMap, VertexShader, PixelShader, bShouldRenderOutlines](FRDGAsyncTask, FRHICommandList& RHICmdList)
 					{
 						// Pass PSO
 						FGraphicsPipelineStateInitializer GraphicsPSOInit;
@@ -491,20 +686,23 @@ void FChannelTextureRenderer::RenderSectionChannels(UTexture* InOutSectionTextur
 						GraphicsPSOInit.BoundShaderState.PixelShaderRHI = PixelShader.GetPixelShader();
 
 
-						GraphicsPSOInit.PrimitiveType = PT_LineList;
-						SetGraphicsPipelineState(RHICmdList, GraphicsPSOInit, 0);
-
-						for (int32 PassId = 0; PassId < Section.Bases_UVBox.Num(); ++PassId)
+						if (bShouldRenderOutlines)
 						{
-							const FInt32Vector2 BaseMeshRange = Section.UVmesh.BaseIndicesRanges[PassId];
+							GraphicsPSOInit.PrimitiveType = PT_LineList;
+							SetGraphicsPipelineState(RHICmdList, GraphicsPSOInit, 0);
 
-							PassParameters->VS.InDrawCall = FUintVector4(Section.UVmesh.NumVertices(), Section.UVmesh.NumChannels(), BaseMeshRange.X, 1);
+							for (int32 PassId = 0; PassId < Section.Bases_UVBox.Num(); ++PassId)
+							{
+								const FInt32Vector2 BaseMeshRange = Section.UVmesh.BaseIndicesRanges[PassId];
 
-							SetShaderParameters(RHICmdList, VertexShader, VertexShader.GetVertexShader(), PassParameters->VS);
-							SetShaderParameters(RHICmdList, PixelShader, PixelShader.GetPixelShader(), PassParameters->PS);
+								PassParameters->VS.InDrawCall = FUintVector4(Section.UVmesh.UVs.Num(), Section.UVmesh.NumChannels(), BaseMeshRange.X, 1);
 
-							int32 NumPrimitives = Section.UVmesh.Outlines.Num();
-							RHICmdList.DrawPrimitive(0, NumPrimitives, 1);
+								SetShaderParameters(RHICmdList, VertexShader, VertexShader.GetVertexShader(), PassParameters->VS);
+								SetShaderParameters(RHICmdList, PixelShader, PixelShader.GetPixelShader(), PassParameters->PS);
+
+								int32 NumPrimitives = Section.UVmesh.Outlines.Num();
+								RHICmdList.DrawPrimitive(0, NumPrimitives, 1);
+							}
 						}
 
 						GraphicsPSOInit.PrimitiveType = PT_TriangleList;
@@ -514,7 +712,7 @@ void FChannelTextureRenderer::RenderSectionChannels(UTexture* InOutSectionTextur
 						{
 							const FInt32Vector2 BaseMeshRange = Section.UVmesh.BaseIndicesRanges[PassId];
 
-							PassParameters->VS.InDrawCall = FUintVector4(Section.UVmesh.NumVertices(), Section.UVmesh.NumChannels(), BaseMeshRange.X, 0);
+							PassParameters->VS.InDrawCall = FUintVector4(Section.UVmesh.UVs.Num(), Section.UVmesh.NumChannels(), BaseMeshRange.X, 0);
 
 							SetShaderParameters(RHICmdList, VertexShader, VertexShader.GetVertexShader(), PassParameters->VS);
 							SetShaderParameters(RHICmdList, PixelShader, PixelShader.GetPixelShader(), PassParameters->PS);
diff --git a/Engine/Plugins/Experimental/MeshPartition/Source/MeshPartitionEditor/Private/MeshPartitionChannelCollection_GenerateChannelUV.cpp b/Engine/Plugins/Experimental/MeshPartition/Source/MeshPartitionEditor/Private/MeshPartitionChannelCollection_GenerateChannelUV.cpp
--- a/Engine/Plugins/Experimental/MeshPartition/Source/MeshPartitionEditor/Private/MeshPartitionChannelCollection_GenerateChannelUV.cpp
+++ b/Engine/Plugins/Experimental/MeshPartition/Source/MeshPartitionEditor/Private/MeshPartitionChannelCollection_GenerateChannelUV.cpp
@@ -1075,14 +1075,186 @@ void FChannelTextureRenderer::GenerateSectionMeshUVs(const FMeshData& InSectionM
 	InSection.SectionIslands_UVBox = { UVDomain };
 	InSection.Bases_UVBox = { UVDomain };
 
+	// FMeshData uses per-vertex UVs (not overlays), so just create a trivial mapping.
+	// If/when FMeshData gets split uv support, this will need to be updated.
+	TArray<int32> UVElementToVertexID;
+	UVElementToVertexID.Reserve(NumVertices);
+	for (int32 i = 0; i < NumVertices; ++i)
+	{
+		UVElementToVertexID.Add(i);
+	}
+
 	InSection.UVmesh.BaseIndicesRanges = { { 0, SectionIndices.Num() } };
 	InSection.UVmesh.Indices = MoveTemp(SectionIndices);
 	InSection.UVmesh.UVs = MoveTemp(SectionUVs);
 	InSection.UVmesh.Outlines = MoveTemp(SectionOutlineEdges);
+	InSection.UVmesh.UVElementToVertexID = MoveTemp(UVElementToVertexID);
+	InSection.UVmesh.VertexCount = NumVertices;
 
 	InSection.RasterResolution = FUint32Vector2((uint32)InOutDomainMapping.ImageResolution);
 	InSection.TexcoordMetrics = FVector2f(InOutDomainMapping.Size3D);
 	InSection.DomainMapping = InOutDomainMapping;
 }
 
+FChannelTextureRenderer::FSection FChannelTextureRenderer::BuildSectionFromCachedTopology(const FChannelRenderUVMeshTopology& InTopology)
+{
+	TRACE_CPUPROFILER_EVENT_SCOPE(BuildSectionFromCachedTopology);
+	
+	FSection OutSection;
+
+	OutSection.UVmesh.Indices = InTopology.Indices;
+	OutSection.UVmesh.UVs = InTopology.UVs;
+	OutSection.UVmesh.UVElementToVertexID = InTopology.UVElementToVertexID;
+	OutSection.UVmesh.VertexCount = InTopology.VertexCount;
+	
+	OutSection.DomainMapping = InTopology.DomainMapping;
+	OutSection.RasterResolution = InTopology.RasterResolution;
+	OutSection.TexcoordMetrics = InTopology.TexcoordMetrics;
+
+	FBox2f UVDomain({ 0,0 }, { 1,1 });
+	OutSection.Islands_Bases = { {0} };
+	OutSection.Islands_UVBox = { UVDomain };
+	OutSection.SectionIslands_UVBox = { UVDomain };
+	OutSection.Bases_UVBox = { UVDomain };
+	
+	OutSection.UVmesh.BaseIndicesRanges = { { 0, InTopology.Indices.Num() } };
+	
+	return OutSection;
+}
+
+MeshPartition::FChannelTextureRenderer::FChannelRenderUVMeshTopology FChannelTextureRenderer::GenerateUVMeshTopology(const Geometry::FDynamicMesh3& InSectionMesh, FVector InScale, float InTexelSize)
+{
+	TRACE_CPUPROFILER_EVENT_SCOPE(GenerateUVMeshTopology);
+
+	FChannelRenderUVMeshTopology Topology;
+
+	const Geometry::FDynamicMeshAttributeSet* Attributes = InSectionMesh.Attributes();
+	if (!Attributes || Attributes->NumUVLayers() == 0)
+	{
+		return Topology;
+	}
+
+	const Geometry::FDynamicMeshUVOverlay* UVOverlay = Attributes->GetUVLayer(0);
+	if (!UVOverlay)
+	{
+		return Topology;
+	}
+
+	// FDynamicMesh3 indices are not necessarily compact:
+	TArray<int32> VertexIDToIndex;
+	{
+		int32 DenseVertexIndex = 0;
+		VertexIDToIndex.SetNum(InSectionMesh.MaxVertexID());
+		for (int VertexID : InSectionMesh.VertexIndicesItr())
+		{
+			VertexIDToIndex[VertexID] = DenseVertexIndex++;
+		}
+	}
+
+	TArray<FVector2f> SectionUVs;
+	TArray<int32> UVElementToVertexID;
+	TArray<int32> UVElementIDToDenseIndex;
+
+	SectionUVs.Reserve(UVOverlay->ElementCount());
+	UVElementToVertexID.Reserve(UVOverlay->ElementCount());
+	UVElementIDToDenseIndex.SetNum(UVOverlay->MaxElementID());
+
+	int32 DenseUVIndex = 0;
+	for (int32 ElementID : UVOverlay->ElementIndicesItr())
+	{
+		FVector2f UV = UVOverlay->GetElement(ElementID);
+		SectionUVs.Add(UV);
+
+		// Get parent vertex ID for this UV element, then convert to dense index
+		int32 ParentVertexID = UVOverlay->GetParentVertex(ElementID);
+		int32 DenseVertexIdx = VertexIDToIndex[ParentVertexID];
+		UVElementToVertexID.Add(DenseVertexIdx);
+
+		UVElementIDToDenseIndex[ElementID] = DenseUVIndex++;
+	}
+
+	TArray<int32> SectionIndices;
+	SectionIndices.Reserve(InSectionMesh.TriangleCount() * 3);
+
+	for (int32 TriangleId : InSectionMesh.TriangleIndicesItr())
+	{
+		Geometry::FIndex3i UVTriangle = UVOverlay->GetTriangle(TriangleId);
+		if (!ensure(UVTriangle.A != Geometry::FDynamicMesh3::InvalidID))
+		{
+			continue;
+		}
+
+		for (int Index = 0; Index < 3; ++Index)
+		{
+			int32 DenseIndex = UVElementIDToDenseIndex[UVTriangle[Index]];
+			UVTriangle[Index] = DenseIndex;
+		}
+		SectionIndices.Append({ UVTriangle.A, UVTriangle.B, UVTriangle.C });
+	}
+
+
+	// Compute domain dimensions
+	FSectionDomainMapping DomainMapping{
+		.TexelSize3D = InTexelSize,
+		.MaxImageResolution = DefaultMaxImageResolution
+	};
+
+	double Section_Area3D = 0;
+	double Section_AreaUV = 0;
+
+	for (int32 TriangleId : InSectionMesh.TriangleIndicesItr())
+	{
+		Geometry::FIndex3i VertexTriangle = InSectionMesh.GetTriangle(TriangleId);
+		Geometry::FIndex3i UVTriangle = UVOverlay->GetTriangle(TriangleId);
+
+		if (UVTriangle.A == Geometry::FDynamicMesh3::InvalidID)
+		{
+			continue; // Skip triangles without UVs
+		}
+
+		// Compute scaled 3D triangle area
+		FVector3d P0 = InSectionMesh.GetVertex(VertexTriangle.A) * InScale;
+		FVector3d P1 = InSectionMesh.GetVertex(VertexTriangle.B) * InScale;
+		FVector3d P2 = InSectionMesh.GetVertex(VertexTriangle.C) * InScale;
+		double TriangleArea3D = 0.5 * FVector3d::CrossProduct((P1 - P0), (P2 - P0)).Length();
+		Section_Area3D += TriangleArea3D;
+
+		// Compute UV triangle area
+		FVector2f UV0 = UVOverlay->GetElement(UVTriangle.A);
+		FVector2f UV1 = UVOverlay->GetElement(UVTriangle.B);
+		FVector2f UV2 = UVOverlay->GetElement(UVTriangle.C);
+		double TriangleAreaUV = 0.5 * FMath::Abs(FVector2f::CrossProduct((UV1 - UV0), (UV2 - UV0)));
+		Section_AreaUV += TriangleAreaUV;
+	}
+
+	// Calculate texture resolution from area ratios
+	double ProjectedSectionDomainArea3D = (Section_AreaUV <= 0 ? 0.0 : Section_Area3D / Section_AreaUV);
+	double SectionSize = FMath::Sqrt(ProjectedSectionDomainArea3D);
+
+	double TexelSize_Unit = DomainMapping.TexelSize3D;
+	double UVScaledImageResolution = SectionSize / TexelSize_Unit;
+	double IdealImageResolution = FMath::Max(DomainMapping.MinImageResolution, UVScaledImageResolution);
+
+	int32 ImageResolution = FMath::RoundUpToPowerOfTwo((int32)FMath::Floor(IdealImageResolution));
+	ImageResolution = FMath::Min(DomainMapping.MaxImageResolution, ImageResolution);
+
+	DomainMapping.Size3D = SectionSize;
+	DomainMapping.SizeUV = SectionSize;
+	DomainMapping.Area3D = Section_Area3D;
+	DomainMapping.AreaUV = Section_AreaUV;
+	DomainMapping.ImageResolution = ImageResolution;
+
+	// Store in topology struct
+	Topology.Indices = MoveTemp(SectionIndices);
+	Topology.UVs = MoveTemp(SectionUVs);
+	Topology.UVElementToVertexID = MoveTemp(UVElementToVertexID);
+	Topology.VertexCount = InSectionMesh.VertexCount();
+
+	Topology.DomainMapping = DomainMapping;
+	Topology.RasterResolution = FUint32Vector2((uint32)ImageResolution);
+	Topology.TexcoordMetrics = FVector2f(DomainMapping.Size3D);
+
+	return Topology;
+}
+
 } // namespace UE::MeshPartition
diff --git a/Engine/Plugins/Experimental/MeshPartition/Source/MeshPartitionEditor/Private/MeshPartitionEditorComponent.cpp b/Engine/Plugins/Experimental/MeshPartition/Source/MeshPartitionEditor/Private/MeshPartitionEditorComponent.cpp
--- a/Engine/Plugins/Experimental/MeshPartition/Source/MeshPartitionEditor/Private/MeshPartitionEditorComponent.cpp
+++ b/Engine/Plugins/Experimental/MeshPartition/Source/MeshPartitionEditor/Private/MeshPartitionEditorComponent.cpp
@@ -1601,7 +1601,10 @@ void UMeshPartitionEditorComponent::FinalizePreviewSectionBuilds()
 			PreviewSection->SetChannelTexture(SectionChannel->Texture.Get());
 			PreviewSection->SetChannelTexcoordDesc(SectionChannel->TexcoordMetrics);
 			PreviewSection->SetChannelTable(SectionChannel->Table);
-			PreviewSection->SetChannelGenerationMesh(SectionChannel->Mesh);
+			if (SectionChannel->Mesh.IsSet())
+			{
+				PreviewSection->SetChannelGenerationMesh(SectionChannel->Mesh.GetValue());
+			}
 
 			PreviewSections.Emplace(PreviewSection);
 
diff --git a/Engine/Plugins/Experimental/MeshPartition/Source/MeshPartitionEditor/Private/MeshPartitionMeshBuilder.cpp b/Engine/Plugins/Experimental/MeshPartition/Source/MeshPartitionEditor/Private/MeshPartitionMeshBuilder.cpp
--- a/Engine/Plugins/Experimental/MeshPartition/Source/MeshPartitionEditor/Private/MeshPartitionMeshBuilder.cpp
+++ b/Engine/Plugins/Experimental/MeshPartition/Source/MeshPartitionEditor/Private/MeshPartitionMeshBuilder.cpp
@@ -383,7 +383,7 @@ namespace UE::MeshPartition::BuildHelpers
 			Hasher.Update(Name.GetData(), Name.Len() * sizeof(TCHAR));
 		};
 
-		static FGuid VersionKey(TEXT("51e0a251-3d12-4852-a6c8-c1b1092faacd"));
+		static FGuid VersionKey(TEXT("03c548f6-a408-4725-961c-020947ede321"));
 		UpdateHashPODType(Hasher, VersionKey);
 
 		// Hash the current mesh data type version into the key to ensure swapping mesh types hits different DDC entries.
diff --git a/Engine/Plugins/Experimental/MeshPartition/Source/MeshPartitionEditor/Public/MeshPartitionChannelCollection.h b/Engine/Plugins/Experimental/MeshPartition/Source/MeshPartitionEditor/Public/MeshPartitionChannelCollection.h
--- a/Engine/Plugins/Experimental/MeshPartition/Source/MeshPartitionEditor/Public/MeshPartitionChannelCollection.h
+++ b/Engine/Plugins/Experimental/MeshPartition/Source/MeshPartitionEditor/Public/MeshPartitionChannelCollection.h
@@ -62,7 +62,7 @@ struct FSectionChannels
 	TArray<uint8> Table;
 	FVector2f TexcoordMetrics = { 1.0f, 1.0f };
 	TStrongObjectPtr<UTexture> Texture;
-	FSectionChannelGenerationMeshData Mesh;
+	TOptional<FSectionChannelGenerationMeshData> Mesh;
 };
 
 /**
@@ -74,6 +74,8 @@ struct FSectionChannels
 class FChannelTextureRenderer
 {
 public:
+	struct FChannelRenderUVMeshTopology;
+	
 	static constexpr int32 NumChannelsPerPage = 1;
 	static constexpr float DefaultTexelSize = 100.0f; // Default texel size (in world unit) driving the channel rasterization resolution
 	static constexpr int32 DefaultMaxImageResolution = 4096; // Default maximum channel texture resolution generated per section
@@ -88,7 +90,15 @@ public:
 	*/
 	static UE_API Tasks::TTask<FSectionChannels> BuildTextureForSection(const FMeshData& InSectionMesh, UObject* InOwner, bool bInDownloadToAsset, const MeshPartition::FChannelMap& InChannels, float InTexelResolution);
 
-private:
+	static UE_API Tasks::TTask<FSectionChannels> BuildTextureForSectionWithCachedTopology(
+		const Geometry::FDynamicMesh3& InSectionMesh,
+		const FChannelRenderUVMeshTopology& InCachedTopology,
+		UObject* InOwner,
+		bool bInDownloadToAsset,
+		const MeshPartition::FChannelMap& InChannels);
+
+	static UE_API FChannelRenderUVMeshTopology GenerateUVMeshTopology(const Geometry::FDynamicMesh3& InSectionMesh, FVector InScale, float InTexelResolution); 
+
 	/**
 	* Section element
 	* A section struct is created for every unique section encountered
@@ -112,10 +122,11 @@ private:
 			TArray<float>			Channels;
 			TArray<int32>			ChannelActiveCounts;
 			TArray<uint8>			ActiveChannelIndices;
+			TArray<int32>			UVElementToVertexID;
+			int32 					VertexCount;
 
 			int32 NumIndices() const { return Indices.Num(); }
-			int32 NumVertices() const { return UVs.Num(); }
-			int32 NumChannels() const { return (NumVertices() ? Channels.Num() / NumVertices() : 0); }
+			int32 NumChannels() const { return (VertexCount ? Channels.Num() / VertexCount : 0); }
 		};
 		FUVMesh					UVmesh;
 
@@ -135,15 +146,37 @@ private:
 		FSectionDomainMapping	DomainMapping;
 	};
 
+	/**
+	 * Cached UV mesh topology that remains constant during displacement and paint operations.
+	 */
+	struct FChannelRenderUVMeshTopology
+	{
+		TArray<int32>       Indices;
+		TArray<FVector2f>   UVs;
+		TArray<int32>       UVElementToVertexID;
+		int32               VertexCount;
+
+		FSectionDomainMapping DomainMapping;
+		FUint32Vector2        RasterResolution;
+		FVector2f             TexcoordMetrics;
+	};
+private:
+
 	static UTexture* GetDefaultTexture();
 
 	static UE_API void GenerateSectionMeshUVs(const FMeshData& InSectionMesh, FSection& InSection, float InTexelResolution);
 	static UE_API void GenerateSectionMeshChannels(const FMeshData& InSectionMesh, FSection& InSection, const MeshPartition::FChannelMap& InChannels);
+	static UE_API void GenerateSectionMeshChannels(const Geometry::FDynamicMesh3& InSectionMesh, FSection& InSection, const MeshPartition::FChannelMap& InChannels);
+
+	static UE_API FSection BuildSectionFromCachedTopology(const FChannelRenderUVMeshTopology& InTopology);
 
 	static UE_API UTexture* AllocateSectionTextureResource(int32 InSizeX, int32 InSizeY, int32 InArraySize, EPixelFormat InFormat, UObject* InOuter, const FName InName);
-	static UE_API void RenderSectionChannels(UTexture* InOutSectionTextures, const FSection& InSection);
+	static UE_API void RenderSectionChannels(UTexture* InOutSectionTextures, FSection InSection);
 	static UE_API void DownloadSectionTextureToAsset(UTexture* InAllocatedTexture);
 
+	static Tasks::TTask<TOptional<UTexture*>> LaunchAllocateTextureTask(Tasks::TTask<FSection> InBuildSectionTask, UObject* InOwner);
+	static Tasks::TTask<FSectionChannels> LaunchRenderTextureTask(Tasks::TTask<FSection> InBuildSectionTask, Tasks::TTask<TOptional<UTexture*>> InAllocateTextureTask, bool bInDownloadToAsset, bool bInExtractUVSectionMeshData = false);
+
 public:
 	static UE_API void CreateSectionUVLayout(FMeshData& InOutSectionMesh, const FChannelCollectionUVLayoutOptions& GenerationOptions);
 
diff --git a/Engine/Plugins/Experimental/MeshPartition/Source/MeshPartitionModelingToolset/Private/MeshPartitionAttributePaintTool.cpp b/Engine/Plugins/Experimental/MeshPartition/Source/MeshPartitionModelingToolset/Private/MeshPartitionAttributePaintTool.cpp
--- a/Engine/Plugins/Experimental/MeshPartition/Source/MeshPartitionModelingToolset/Private/MeshPartitionAttributePaintTool.cpp
+++ b/Engine/Plugins/Experimental/MeshPartition/Source/MeshPartitionModelingToolset/Private/MeshPartitionAttributePaintTool.cpp
@@ -6,6 +6,7 @@
 #include "MeshPartitionModifierComponent.h"
 #include "ModelingToolTargetUtil.h"
 #include "Modifiers/MeshPartitionProjectSculptLayersModifier.h"
+#include "MeshPartitionModifierToolTarget.h"
 
 #define LOCTEXT_NAMESPACE "UAttributePaintTool"
 
@@ -119,6 +120,8 @@ void UAttributePaintTool::Setup()
 	bPreferMeshDescription = false;
 
 	Super::Setup();
+	// Mesh Partition tools all should default to showing the existing material.
+	ViewProperties->MaterialMode = EMeshAttributePaintMaterialMode::ExistingMaterial;
 
 	// This tool works on dynamic meshes only, currently
 	ensure(EditedDynamicMesh);
@@ -148,6 +151,12 @@ void UAttributePaintTool::Setup()
 			PendingNewSelectedIndex = AttribProps->Attributes.Num() > 0 ? 0 : -1;
 		}
 	}
+
+	if (UEditableModifierToolTarget* ModifierTarget = Cast<UEditableModifierToolTarget>(Target))
+	{
+		ModifierTarget->UpdateRenderTextureForPreview(*EditedDynamicMesh);
+		ModifierTarget->ConfigurePreviewForRendering(PreviewMesh->GetRootComponent());
+	}
 }
 
 void UAttributePaintTool::OnShutdown(EToolShutdownType ShutdownType)
@@ -197,6 +206,28 @@ void UAttributePaintTool::OnShutdown(EToolShutdownType ShutdownType)
 	bFilteredAttributeSourceInitialized = false;
 }
 
+void UAttributePaintTool::ApplyStamp(const FBrushStampData& Stamp)
+{
+	Super::ApplyStamp(Stamp);
+
+	if (UEditableModifierToolTarget* ModifierTarget = Cast<UEditableModifierToolTarget>(Target))
+	{
+		ModifierTarget->UpdateRenderTextureForPreview(*EditedDynamicMesh);
+		ModifierTarget->ConfigurePreviewForRendering(PreviewMesh->GetRootComponent());
+	}
+}
+
+void UAttributePaintTool::ExternalUpdateValues(int32 AttribIndex, const TArray<int32>& VertexIndices, const TArray<float>& NewValues)
+{
+	Super::ExternalUpdateValues(AttribIndex, VertexIndices, NewValues);
+
+	if (UEditableModifierToolTarget* ModifierTarget = Cast<UEditableModifierToolTarget>(Target))
+	{
+		ModifierTarget->UpdateRenderTextureForPreview(*EditedDynamicMesh);
+		ModifierTarget->ConfigurePreviewForRendering(PreviewMesh->GetRootComponent());
+	}
+}
+
 
 void UAttributePaintTool::AddChannel(FName NewChannel)
 {
diff --git a/Engine/Plugins/Experimental/MeshPartition/Source/MeshPartitionModelingToolset/Private/MeshPartitionModifierToolTarget.cpp b/Engine/Plugins/Experimental/MeshPartition/Source/MeshPartitionModelingToolset/Private/MeshPartitionModifierToolTarget.cpp
--- a/Engine/Plugins/Experimental/MeshPartition/Source/MeshPartitionModelingToolset/Private/MeshPartitionModifierToolTarget.cpp
+++ b/Engine/Plugins/Experimental/MeshPartition/Source/MeshPartitionModelingToolset/Private/MeshPartitionModifierToolTarget.cpp
@@ -4,6 +4,7 @@
 #include "MeshPartitionModifierToolTarget.h"
 
 #include "DynamicMeshEditor.h"
+#include "MeshPartitionChannelCollection.h"
 #include "MeshPartitionComponent.h"
 #include "MeshPartitionDefinition.h"
 #include "MeshPartitionEditorComponent.h"
@@ -146,6 +147,42 @@ void UModifierToolTarget::ConfigurePreviewForRendering(UPrimitiveComponent* Prim
 	PrimitiveComponent->SetDefaultCustomPrimitiveDataVector2f(DataIndex, ChannelTexcoordDesc);
 }
 
+void UModifierToolTarget::UpdateRenderTextureForPreview(const FDynamicMesh3& InPreviewMesh)
+{
+	if (UMeshPartitionEditorComponent* EditorComponent = Cast<UMeshPartitionEditorComponent>(TargetModifier->GetMeshPartitionEditorComponent()))
+	{
+		const UMeshPartitionDefinition* Definition = EditorComponent->GetMegaMeshDefinition() ? EditorComponent->GetMegaMeshDefinition() : UMeshPartitionDefinition::GetDefaultMegaMeshDefinition();
+
+		const MeshPartition::FChannelMap ChannelMap = Definition->GetChannelMap();
+		const float ChannelTexelSize = Definition->GetChannelTexelSize();
+
+		constexpr bool bDownloadToAsset = false;
+
+		if (!CachedUVTopology.IsSet())
+		{
+			const FTransformSRT3d ModifierTransform = (FTransformSRT3d)GetWorldTransform();
+			CachedUVTopology = FChannelTextureRenderer::GenerateUVMeshTopology(InPreviewMesh, ModifierTransform.GetScale3D(), ChannelTexelSize);
+		}
+
+		Tasks::TTask<MeshPartition::FSectionChannels> RenderSectionChannels = FChannelTextureRenderer::BuildTextureForSectionWithCachedTopology(
+			InPreviewMesh,
+			CachedUVTopology.GetValue(),
+			this,
+			bDownloadToAsset,
+			ChannelMap);
+
+		RenderSectionChannels.Wait();
+
+		MeshPartition::FSectionChannels SectionChannels = MoveTemp(RenderSectionChannels.GetResult());
+		ChannelTable = MoveTemp(SectionChannels.Table);
+		ChannelTexture = SectionChannels.Texture.Get();
+		ChannelTexcoordDesc = SectionChannels.TexcoordMetrics;
+
+		MID = MeshPartition::EditorUtils::GetOrCreateMaterialInstance(MID, EditorComponent->GetDefinitionMaterial(), this, TEXT("ModifierToolTargetMID"), RF_Transient);
+		MID->SetTextureParameterValue(UE::MeshPartition::ChannelTextureParameterName, ChannelTexture);
+	}
+}
+
 bool UModifierToolTarget::IsValid() const
 {
 	return TargetModifier && TargetModifier->GetMeshPartitionEditorComponent();
@@ -229,27 +266,6 @@ TArray<MeshPartition::UModifierComponent*> UModifierToolTarget::GetModifiersToPr
 	return ModifiersToProcess;
 }
 
-void UModifierToolTarget::RenderChannelTextures(const FMeshData& InRenderMesh)
-{
-	if (UMeshPartitionEditorComponent* EditorComponent = Cast<UMeshPartitionEditorComponent>(TargetModifier->GetMeshPartitionEditorComponent()))
-	{
-		const UMeshPartitionDefinition* Definition = EditorComponent->GetMegaMeshDefinition() ? EditorComponent->GetMegaMeshDefinition() : UMeshPartitionDefinition::GetDefaultMegaMeshDefinition();
-
-		const MeshPartition::FChannelMap ChannelMap = Definition->GetChannelMap();
-		const float ChannelTexelSize = Definition->GetChannelTexelSize();
-
-		constexpr bool bDownloadToAsset = false;
-		Tasks::TTask<MeshPartition::FSectionChannels> RenderSectionChannels = FChannelTextureRenderer::BuildTextureForSection(InRenderMesh, this, bDownloadToAsset, ChannelMap, ChannelTexelSize);
-		MeshPartition::FSectionChannels SectionChannels = MoveTemp(RenderSectionChannels.GetResult());
-		ChannelTable = MoveTemp(SectionChannels.Table);
-		ChannelTexture = SectionChannels.Texture.Get();
-		ChannelTexcoordDesc = SectionChannels.TexcoordMetrics;
-
-		MID = MeshPartition::EditorUtils::GetOrCreateMaterialInstance(MID, EditorComponent->GetDefinitionMaterial(), this, TEXT("ModifierToolTargetMID"), RF_Transient);
-		MID->SetTextureParameterValue(UE::MeshPartition::ChannelTextureParameterName, ChannelTexture);
-	}
-}
-
 // -------------------------------------------------------------------------------------------------------------------------
 
 FDynamicMesh3 UEditableModifierToolTarget::GetDynamicMesh()
@@ -266,11 +282,7 @@ FDynamicMesh3 UEditableModifierToolTarget::GetDynamicMesh()
 			EditableModifier->PrepareForEdit(Mesh);
 		}
 
-		// #todo: there's an opportunity to speed up the performance here by implementing a path to directly render the dynamic mesh into a channel texture.
-		// This is also technically incorrect, since conversion from DynamicMesh -> MeshData does not currently generate split vertices when required to support split attributes.
-		FMeshData MeshData;
-		MeshData.AppendDynamicMesh(Mesh, FTransform::Identity);
-		RenderChannelTextures(MeshData);
+		UpdateRenderTextureForPreview(Mesh);
 	}
 	return Mesh;
 }
diff --git a/Engine/Plugins/Experimental/MeshPartition/Source/MeshPartitionModelingToolset/Public/MeshPartitionAttributePaintTool.h b/Engine/Plugins/Experimental/MeshPartition/Source/MeshPartitionModelingToolset/Public/MeshPartitionAttributePaintTool.h
--- a/Engine/Plugins/Experimental/MeshPartition/Source/MeshPartitionModelingToolset/Public/MeshPartitionAttributePaintTool.h
+++ b/Engine/Plugins/Experimental/MeshPartition/Source/MeshPartitionModelingToolset/Public/MeshPartitionAttributePaintTool.h
@@ -43,6 +43,10 @@ public:
 	// UDynamicMeshBrushTool
 	virtual void OnShutdown(EToolShutdownType ShutdownType) override;
 
+protected:
+	virtual void ExternalUpdateValues(int32 AttribIndex, const TArray<int32>& VertexIndices, const TArray<float>& NewValues) override;
+	virtual void ApplyStamp(const FBrushStampData& Stamp);
+
 private:
 	UPROPERTY()
 	TObjectPtr<MeshPartition::UAttributePaintToolAddChannelProperties> AddChannelProperties = nullptr;
diff --git a/Engine/Plugins/Experimental/MeshPartition/Source/MeshPartitionModelingToolset/Public/MeshPartitionModifierToolTarget.h b/Engine/Plugins/Experimental/MeshPartition/Source/MeshPartitionModelingToolset/Public/MeshPartitionModifierToolTarget.h
--- a/Engine/Plugins/Experimental/MeshPartition/Source/MeshPartitionModelingToolset/Public/MeshPartitionModifierToolTarget.h
+++ b/Engine/Plugins/Experimental/MeshPartition/Source/MeshPartitionModelingToolset/Public/MeshPartitionModifierToolTarget.h
@@ -12,6 +12,7 @@
 #include "ToolTargets/SceneComponentToolTarget.h"
 #include "UObject/Object.h"
 #include "MeshPartitionChannel.h"
+#include "MeshPartitionChannelCollection.h"
 
 #include "MeshPartitionModifierToolTarget.generated.h"
 
@@ -20,6 +21,11 @@
 class UTexture;
 class UMaterialInstanceDynamic;
 
+namespace UE::Geometry
+{
+class FDynamicMesh3;
+}
+
 namespace UE::MeshPartition
 {
 class UMeshPartitionEditorComponent;
@@ -64,13 +70,13 @@ public:
 	UE_API virtual void SetOwnerVisibility(bool bInVisible) const override;
 	
 	void ConfigurePreviewForRendering(UPrimitiveComponent* PrimitiveComponent) const;
+
+	void UpdateRenderTextureForPreview(const Geometry::FDynamicMesh3& PreviewMesh);
 protected:
 	UE_API void BuildModifiedMeshUpToTarget(Geometry::FDynamicMesh3& OutResultMesh, bool bIncludeTargetModifier = false);
 
 	UE_API virtual TArray<FBox> GetBounds() const;
 	UE_API virtual TArray<MeshPartition::UModifierComponent*> GetModifiersToProcess() const;
-	
-	void RenderChannelTextures(const FMeshData& InRenderMesh);
 
 protected:
 	UPROPERTY()
@@ -86,6 +92,9 @@ protected:
 	UPROPERTY()
 	TObjectPtr<UMaterialInstanceDynamic> MID;
 
+	/** Cached UV topology for fast texture updates during painting */
+	TOptional<MeshPartition::FChannelTextureRenderer::FChannelRenderUVMeshTopology> CachedUVTopology;
+
 private:
 	// When using SetOwnerVisibility, we hide all base sections that were entangled with our preview sections. This
 	//  keeps track of them so that we can unhide them.
