Avoid delete-delete version control conflicts in the mesh partition builder.

When someone else already has the file checked out (i.e. deleted by running the mesh partition builder),
then we can skip the deletion ourselves to avoid any conflict errors.
Old files left around are handled properly, they will be ignored, until someone checks in the deletion.

#Synced-CL 51480380

diff --git a/Engine/Plugins/Experimental/MeshPartition/Source/MeshPartitionEditor/Private/WorldPartitionMeshPartitionBuilder.cpp b/Engine/Plugins/Experimental/MeshPartition/Source/MeshPartitionEditor/Private/WorldPartitionMeshPartitionBuilder.cpp
--- a/Engine/Plugins/Experimental/MeshPartition/Source/MeshPartitionEditor/Private/WorldPartitionMeshPartitionBuilder.cpp
+++ b/Engine/Plugins/Experimental/MeshPartition/Source/MeshPartitionEditor/Private/WorldPartitionMeshPartitionBuilder.cpp
@@ -47,6 +47,7 @@
 #include "WorldPartition/DataLayer/DataLayerAsset.h"
 #include "WorldPartition/DataLayer/DataLayerManager.h"
 #include "WorldPartition/DataLayer/DataLayerInstance.h"
+#include "SourceControlOperations.h"
 
 #define LOCTEXT_NAMESPACE "MeshPartitionBuilder"
 
@@ -157,9 +158,10 @@ bool UWorldPartitionMeshPartitionBuilder::RunInternal(UWorld* InWorld, const FCe
 
 	TArray<UPackage*> PackagesToCleanup;
 	int32 TotalRemovedCompiledSections = 0;
+	int32 TotalFailedToRemoveCompiledSections = 0;
 	TArray<TWeakObjectPtr<AActor>> ActorsToUnload;
 	FWorldPartitionHelpers::ForEachActorDescInstance(WorldPartition, MeshPartition::ACompiledSection::StaticClass(),
-	[InWorld, WorldPartition = WorldPartition, BuildKey, CompiledSectionActorDescGuidsToKeep = CompiledSectionActorDescGuidsToKeep, &TotalRemovedCompiledSections, &PackagesToCleanup, &ActorsToUnload](const FWorldPartitionActorDescInstance* ActorDescInstance)
+	[InWorld, WorldPartition = WorldPartition, BuildKey, CompiledSectionActorDescGuidsToKeep = CompiledSectionActorDescGuidsToKeep, &TotalRemovedCompiledSections, &TotalFailedToRemoveCompiledSections, &PackagesToCleanup, &ActorsToUnload](const FWorldPartitionActorDescInstance* ActorDescInstance)
 	{
 		const FGuid& ActorDescGuid = ActorDescInstance->GetGuid();
 		if (CompiledSectionActorDescGuidsToKeep.Contains(ActorDescGuid))
@@ -198,6 +200,32 @@ bool UWorldPartitionMeshPartitionBuilder::RunInternal(UWorld* InWorld, const FCe
 					);
 
 					check(BuildInfo.BuildKey != BuildKey); // shouldn't be deleting any actors that were just built
+
+					// before deleting, check if the package is exclusively checked out by someone else in source control.
+					// If so, skip deletion to avoid a conflict
+					// (we don't require deletion, the system can filter out old out-of-date files)
+					// If performance becomes a concern, the file states could be batch-prefetched with SCCProvider.UpdateStatus(...) before the ForEachActorDescInstance call and then use EStateCacheUsage::Use here.
+					{
+						const FString PackageFilename = SourceControlHelpers::PackageFilename(CompiledSection->GetPackage());
+						ISourceControlProvider& SCCProvider = ISourceControlModule::Get().GetProvider();
+						if (SCCProvider.IsEnabled())
+						{
+							FSourceControlStatePtr SCCState = SCCProvider.GetState(PackageFilename, EStateCacheUsage::ForceUpdate);
+							if (SCCState.IsValid())
+							{
+								if (SCCState->IsCheckedOutOther())
+								{
+									// don't attempt to delete if it 
+									UE_LOG(LogMegaMeshEditor, Display,
+										TEXT("  Skipping deletion of compiled section %.6s: package '%s' is exclusively checked out by another user (%s)."),
+										*ActorDescGuid.ToString(), *PackageFilename, *SCCState->GetDisplayName().ToString());
+									TotalFailedToRemoveCompiledSections++;
+									return true;
+								}
+							}
+						}
+					}
+
 					PackagesToCleanup.Emplace(CompiledSection->GetPackage());
 					ActorsToUnload.Add(CompiledSection);
 					InWorld->DestroyActor(CompiledSection);
@@ -243,22 +271,24 @@ bool UWorldPartitionMeshPartitionBuilder::RunInternal(UWorld* InWorld, const FCe
 		}
 	}
 
-	UE_LOG(LogMegaMeshEditor, Display, TEXT("Finished Build for World %s (Built %d, Reused %d, Failed %d, Removed %d Compiled Sections)"), *InWorld->GetName(),
+	UE_LOG(LogMegaMeshEditor, Display, TEXT("Finished Build for World %s (Built %d, Reused %d, Failed %d, Removed %d, Could Not Remove %d Compiled Sections)"), *InWorld->GetName(),
 		TotalBuiltCompiledSections,
 		TotalReusedCompiledSections,
 		TotalFailedCompiledSections,
-		TotalRemovedCompiledSections);
+		TotalRemovedCompiledSections,
+		TotalFailedToRemoveCompiledSections);
 
 	if ((FailedMeshPartitionBuilds == 0) && (TotalFailedCompiledSections == 0))
 	{
 		return true;
 	}
 
-	UE_LOG(LogMegaMeshEditor, Warning, TEXT("Failed to build for World %s (Built %d, Reused %d, Failed %d, Removed %d Compiled Sections), %d Mesh Partitions Failed"), *InWorld->GetName(),
+	UE_LOG(LogMegaMeshEditor, Warning, TEXT("Failed to build for World %s (Built %d, Reused %d, Failed %d, Removed %d, Could Not Remove %d Compiled Sections), %d Mesh Partitions Failed"), *InWorld->GetName(),
 		TotalBuiltCompiledSections,
 		TotalReusedCompiledSections,
 		TotalFailedCompiledSections,
 		TotalRemovedCompiledSections,
+		TotalFailedToRemoveCompiledSections,
 		FailedMeshPartitionBuilds);
 
 	return false;
