make meshpartition tool target fall back to splitting combined mesh by closest modifier bounds in case where the section id attribute is lost/destroyed by the tool

#Synced-CL 51735880

diff --git a/Engine/Plugins/Experimental/MeshPartition/Source/MeshPartitionModelingToolset/Private/MeshPartitionMultiSectionToolTarget.cpp b/Engine/Plugins/Experimental/MeshPartition/Source/MeshPartitionModelingToolset/Private/MeshPartitionMultiSectionToolTarget.cpp
--- a/Engine/Plugins/Experimental/MeshPartition/Source/MeshPartitionModelingToolset/Private/MeshPartitionMultiSectionToolTarget.cpp
+++ b/Engine/Plugins/Experimental/MeshPartition/Source/MeshPartitionModelingToolset/Private/MeshPartitionMultiSectionToolTarget.cpp
@@ -3,6 +3,7 @@
 #include "MeshPartitionMultiSectionToolTarget.h"
 
 #include "Algo/Transform.h"
+#include "Async/ParallelFor.h"
 #include "DynamicMesh/DynamicMesh3.h"
 #include "DynamicMesh/DynamicMeshAttributeSet.h"
 #include "DynamicMesh/DynamicMeshTriangleAttribute.h"
@@ -11,6 +12,7 @@
 #include "HAL/IConsoleManager.h" // FAutoConsoleVariableRef
 #include "MeshPartitionEditorComponent.h"
 #include "MeshPartitionModifierUtils.h" // IsNonzeroWeightLayer
+#include "MeshPartitionToolTargetUtils.h"
 #include "MeshPartitionPreviewSection.h"
 #include "Modifiers/MeshPartitionMeshProvider.h"
 #include "ModelingToolTargetUtil.h"
@@ -187,28 +189,46 @@ void UMultiSectionToolTarget::SplitMeshesToSections(const Geometry::FDynamicMesh
 		return;
 	}
 
-	if (!ensure(AccumulatedMesh.HasAttributes()))
-	{
-		return;
-	}
 
-	const TDynamicMeshScalarTriangleAttribute<int32>* SectionIDs = static_cast<const TDynamicMeshScalarTriangleAttribute<int32>*>(
-		AccumulatedMesh.Attributes()->GetAttachedAttribute(SectionIDOverlayAttributeName));
+	const TDynamicMeshScalarTriangleAttribute<int32>* SectionIDs = (AccumulatedMesh.HasAttributes())
+		? static_cast<const TDynamicMeshScalarTriangleAttribute<int32>*>(AccumulatedMesh.Attributes()->GetAttachedAttribute(SectionIDOverlayAttributeName))
+		: nullptr;
 
-	if (!ensure(SectionIDs))
+	TArray<int32> FallbackSectionIDs;
+	if (!SectionIDs)
 	{
-		return;
+		// tool lost/destroyed our section ID attribute -- fall back to assigning triangles by source modifier bounds
+		TArray<FAxisAlignedBox3d> SectionBounds;
+		SectionBounds.Init(FAxisAlignedBox3d::Empty(), BaseSections.Num());
+		ParallelFor(BaseSections.Num(), [&](int32 SectionIndex)
+		{
+			const UMeshProviderModifier* Section = BaseSections[SectionIndex].Get();
+			if (!Section || !Section->GetMesh())
+			{
+				return;
+			}
+			const FDynamicMesh3* ModifierMesh = Section->GetMesh();
+			const FTransform& SectionToWorld = Section->GetComponentToWorld();
+			for (int32 VID : ModifierMesh->VertexIndicesItr())
+			{
+				SectionBounds[SectionIndex].Contain(
+					TargetTransform.InverseTransformPosition(
+						SectionToWorld.TransformPosition(FVector(ModifierMesh->GetVertex(VID))
+				)) );
+			}
+		});
+
+		FallbackSectionIDs = AssignMeshTrisToClosestBounds(AccumulatedMesh, SectionBounds);
 	}
 
 	TArray<TUniquePtr<FDynamicMesh3>> SplitMeshes;
 	TArray<int32> SplitIdxToSectionIdx, SectionIdxToSplitIdx;
 	constexpr bool bSplitIfSingle = true;
 	const bool bWasSplit = FDynamicMeshEditor::SplitMesh(&AccumulatedMesh, SplitMeshes, bSplitIfSingle,
-		[SectionIDs](int TID)
+		[SectionIDs, &FallbackSectionIDs](int TID) -> int32
 		{
-			const int SectionID = SectionIDs->GetValue(TID);
-			return SectionID;
-		}, 
+			return SectionIDs ? SectionIDs->GetValue(TID) : FallbackSectionIDs[TID];
+		},
 		// DeleteMeshID, not used
 		-1,
 		// Will be 1:1 with SplitMeshes
diff --git a/Engine/Plugins/Experimental/MeshPartition/Source/MeshPartitionModelingToolset/Private/MeshPartitionToolTarget.cpp b/Engine/Plugins/Experimental/MeshPartition/Source/MeshPartitionModelingToolset/Private/MeshPartitionToolTarget.cpp
--- a/Engine/Plugins/Experimental/MeshPartition/Source/MeshPartitionModelingToolset/Private/MeshPartitionToolTarget.cpp
+++ b/Engine/Plugins/Experimental/MeshPartition/Source/MeshPartitionModelingToolset/Private/MeshPartitionToolTarget.cpp
@@ -3,7 +3,9 @@
 
 #include "MeshPartitionToolTarget.h"
 
+#include "Async/ParallelFor.h"
 #include "DynamicMeshEditor.h"
+#include "MeshPartitionToolTargetUtils.h"
 #include "MeshPartitionComponent.h"
 #include "ModelingToolTargetUtil.h"
 #include "Modifiers/MeshPartitionMeshProvider.h"
@@ -156,20 +158,48 @@ void UMeshPartitionToolTarget::SplitMeshesToSections(const FDynamicMesh3& InComb
 {
 	TRACE_CPUPROFILER_EVENT_SCOPE(MeshPartition::UMeshPartitionToolTarget::SplitMeshesToSections);
 
-	const TDynamicMeshScalarTriangleAttribute<int32>* SectionIDs = static_cast<const TDynamicMeshScalarTriangleAttribute<int32>*>(InCombinedMesh.Attributes()->GetAttachedAttribute(SectionIDOverlayAttributeName));
+	const TDynamicMeshScalarTriangleAttribute<int32>* SectionIDs = (InCombinedMesh.Attributes())
+		? static_cast<const TDynamicMeshScalarTriangleAttribute<int32>*>(InCombinedMesh.Attributes()->GetAttachedAttribute(SectionIDOverlayAttributeName))
+		: nullptr;
 
 	check(Component.IsValid());
+	check(bIsMeshModifiersInitialized);
 
 	TArray<TUniquePtr<FDynamicMesh3>> SplitMeshes;
 	TArray<int32> SplitIdxToSectionIdx, SectionIdxToSplitIdx;
 	constexpr int32 bDeleteID = INDEX_NONE;
 	constexpr bool bSplitSingleMesh = true;
-	const bool bWasSplit = FDynamicMeshEditor::SplitMesh(&InCombinedMesh, SplitMeshes, bSplitSingleMesh, [SectionIDs](int TID)
+
+	TArray<int32> FallbackSectionIDs;
+	if (!SectionIDs)
 	{
-		const int SectionID = SectionIDs->GetValue(TID);
-		check(SectionID != INDEX_NONE);
-		return SectionID;
-	}, bDeleteID, &SplitIdxToSectionIdx);
+		// tool lost/destroyed our section ID attribute -- fall back to assigning triangles by source modifier bounds
+		const FTransform3d TargetToWorld = Component->GetComponentToWorld();
+		TArray<FAxisAlignedBox3d> SectionBounds;
+		SectionBounds.Init(FAxisAlignedBox3d::Empty(), MeshModifiers.Num());
+		ParallelFor(MeshModifiers.Num(), [&](int32 SectionIndex)
+		{
+			if (!MeshModifiers[SectionIndex] || !MeshModifiers[SectionIndex]->GetMesh())
+			{
+				return;
+			}
+			const FDynamicMesh3* ModifierMesh = MeshModifiers[SectionIndex]->GetMesh();
+			// Note we're matching BuildMergedMesh's vertex transform here
+			const FTransform3d SectionToTarget = TargetToWorld.Inverse() * MeshModifiers[SectionIndex]->GetComponentTransform();
+			for (int32 VID : ModifierMesh->VertexIndicesItr())
+			{
+				SectionBounds[SectionIndex].Contain(SectionToTarget.TransformPosition(ModifierMesh->GetVertex(VID)));
+			}
+		});
+
+		FallbackSectionIDs = AssignMeshTrisToClosestBounds(InCombinedMesh, SectionBounds);
+	}
+
+	const bool bWasSplit = FDynamicMeshEditor::SplitMesh(&InCombinedMesh, SplitMeshes, bSplitSingleMesh,
+		[SectionIDs, &FallbackSectionIDs](int TID) -> int32
+		{
+			return SectionIDs ? SectionIDs->GetValue(TID) : FallbackSectionIDs[TID];
+		}, bDeleteID, &SplitIdxToSectionIdx);
 
 	check(bWasSplit);
 	// build reverse map
@@ -183,9 +213,6 @@ void UMeshPartitionToolTarget::SplitMeshesToSections(const FDynamicMesh3& InComb
 		}
 	}
 
-
-	check(bIsMeshModifiersInitialized);
-
 	// Accumulate the new bounds of all the sections so we can trigger a single OnModifierChanged once all sections are updated
 	TArray<FBox> AccumulatedBounds;
 	for (int32 SectionIndex = 0; SectionIndex < MeshModifiers.Num(); ++SectionIndex)
diff --git a/Engine/Plugins/Experimental/MeshPartition/Source/MeshPartitionModelingToolset/Private/MeshPartitionToolTargetUtils.cpp b/Engine/Plugins/Experimental/MeshPartition/Source/MeshPartitionModelingToolset/Private/MeshPartitionToolTargetUtils.cpp
new file mode 100644
--- /dev/null
+++ b/Engine/Plugins/Experimental/MeshPartition/Source/MeshPartitionModelingToolset/Private/MeshPartitionToolTargetUtils.cpp
@@ -0,0 +1,49 @@
+// Copyright Epic Games, Inc. All Rights Reserved.
+
+#include "MeshPartitionToolTargetUtils.h"
+
+#include "Async/ParallelFor.h"
+
+namespace UE::MeshPartition
+{
+
+TArray<int32> AssignMeshTrisToClosestBounds(
+	const UE::Geometry::FDynamicMesh3& Mesh,
+	const TArray<UE::Geometry::FAxisAlignedBox3d>& Bounds,
+	int32 DefaultAssignment)
+{
+	using namespace UE::Geometry;
+
+	TArray<int32> SectionIDs;
+	SectionIDs.SetNumUninitialized(Mesh.MaxTriangleID());
+
+	ParallelFor(Mesh.MaxTriangleID(), [&](int32 TID)
+	{
+		if (!Mesh.IsTriangle(TID))
+		{
+			SectionIDs[TID] = INDEX_NONE;
+			return;
+		}
+		const FVector3d Centroid = Mesh.GetTriCentroid(TID);
+
+		int32 BestSection = DefaultAssignment;
+		double BestDistSq = TNumericLimits<double>::Max();
+		for (int32 SectionIndex = 0; SectionIndex < Bounds.Num(); ++SectionIndex)
+		{
+			if (!Bounds[SectionIndex].IsEmpty())
+			{
+				const double DistSq = Bounds[SectionIndex].DistanceSquared(Centroid);
+				if (DistSq < BestDistSq)
+				{
+					BestDistSq = DistSq;
+					BestSection = SectionIndex;
+				}
+			}
+		}
+		SectionIDs[TID] = BestSection;
+	});
+
+	return SectionIDs;
+}
+
+} // namespace UE::MeshPartition
diff --git a/Engine/Plugins/Experimental/MeshPartition/Source/MeshPartitionModelingToolset/Private/MeshPartitionToolTargetUtils.h b/Engine/Plugins/Experimental/MeshPartition/Source/MeshPartitionModelingToolset/Private/MeshPartitionToolTargetUtils.h
new file mode 100644
--- /dev/null
+++ b/Engine/Plugins/Experimental/MeshPartition/Source/MeshPartitionModelingToolset/Private/MeshPartitionToolTargetUtils.h
@@ -0,0 +1,22 @@
+// Copyright Epic Games, Inc. All Rights Reserved.
+
+#pragma once
+
+#include "BoxTypes.h"
+#include "DynamicMesh/DynamicMesh3.h"
+
+namespace UE::MeshPartition
+{
+
+/**
+ * Helper method used to split triangles out to sections, via bounding boxes
+ * 
+ * @param DefaultAssignment If no bounds are valid, all triangles will be assigned this value.
+ * @return Array mapping triangle IDs to the index of the triangle centroid's closest Bounds
+ */
+TArray<int32> AssignMeshTrisToClosestBounds(
+	const UE::Geometry::FDynamicMesh3& Mesh,
+	const TArray<UE::Geometry::FAxisAlignedBox3d>& Bounds,
+	int32 DefaultAssignment = 0);
+
+} // namespace UE::MeshPartition
