Added 'trim' option to mesh partition's boolean modifier, which can be used to cut an open hole in the mesh partition surface

#Synced-CL 51921100

diff --git a/Engine/Plugins/Experimental/MeshPartition/Source/MeshPartitionEditor/Private/Modifiers/MeshPartitionBooleanModifier.cpp b/Engine/Plugins/Experimental/MeshPartition/Source/MeshPartitionEditor/Private/Modifiers/MeshPartitionBooleanModifier.cpp
--- a/Engine/Plugins/Experimental/MeshPartition/Source/MeshPartitionEditor/Private/Modifiers/MeshPartitionBooleanModifier.cpp
+++ b/Engine/Plugins/Experimental/MeshPartition/Source/MeshPartitionEditor/Private/Modifiers/MeshPartitionBooleanModifier.cpp
@@ -235,9 +235,24 @@ namespace MegaMeshBooleanLocals
 		MeshNormals.ComputeVertexNormals(false, true); // angle weighted for more robust inside/outside 
 		MeshNormals.CopyToVertexNormals(&EditableMesh);
 
-		const FMeshBoolean::EBooleanOp BooleanOpMode = BooleanOp == MeshPartition::EBooleanOperation::Union
-			? FMeshBoolean::EBooleanOp::Union
-			: FMeshBoolean::EBooleanOp::Difference;
+		FMeshBoolean::EBooleanOp BooleanOpMode;
+		bool bNeedsEmbeddingLogic = true;
+		switch (BooleanOp)
+		{
+		default:
+			ensureMsgf(false, TEXT("Unsupported BooleanOp in Mesh Partition Boolean Modifier"));
+
+		case MeshPartition::EBooleanOperation::Union:
+			BooleanOpMode = FMeshBoolean::EBooleanOp::Union;
+			break;
+		case MeshPartition::EBooleanOperation::Subtract:
+			BooleanOpMode = FMeshBoolean::EBooleanOp::Difference;
+			break;
+		case MeshPartition::EBooleanOperation::Trim:
+			BooleanOpMode = FMeshBoolean::EBooleanOp::TrimInside;
+			bNeedsEmbeddingLogic = false;
+			break;
+		}
 
 		// passing in the submesh first seems to preserve boundary vertices, but it needs to be verified if this is guaranteed or
 		// whether extra work to remap them is needed
@@ -256,7 +271,11 @@ namespace MegaMeshBooleanLocals
 		MeshBooleanOp.WindingThreshold = ToolWindingThreshold;
 		MeshBooleanOp.bWeldSharedEdges = bWeldSharedEdges;
 		MeshBooleanOp.bSimplifyAlongNewEdges = bSimplifyAlongNewEdges;
-		MeshBooleanOp.TrackPerTriangleSourceMesh.Emplace();
+		if (!FMeshBoolean::OperatesOnSingleMesh(BooleanOpMode))
+		{
+			MeshBooleanOp.TrackPerTriangleSourceMesh.Emplace();
+		}
+		
 
 		FMeshBoolean::FCustomInsideMeshTest ToolWindingTest(
 			[this](const FVector3d& Pos, const FMeshBoolean::FCustomInsideTestContext& Context) -> bool
@@ -264,16 +283,19 @@ namespace MegaMeshBooleanLocals
 				return Context.OptionalWindingTree->IsInside(Pos, ToolWindingThreshold);
 			}, true /*bRequiresWinding*/);
 		double UseTargetWindingThreshold = TargetWindingThreshold;
-		if (ToolMeshEmbedding == MeshPartition::EBooleanToolMeshEmbedding::Intersecting)
+		if (bNeedsEmbeddingLogic)
 		{
-			// The missing half-enclosure would have added ~.5 to the winding number.
-			// In practice, subtracting slightly less than .5 helps give expected results even when the tool mesh is outside of the terrain.
-			UseTargetWindingThreshold -= .49;
-		}
-		else if (ToolMeshEmbedding == MeshPartition::EBooleanToolMeshEmbedding::Inside)
-		{
-			// The missing full-enclosure would have added 1 to the winding number.
-			UseTargetWindingThreshold -= 1.;
+			if (ToolMeshEmbedding == MeshPartition::EBooleanToolMeshEmbedding::Intersecting)
+			{
+				// The missing half-enclosure would have added ~.5 to the winding number.
+				// In practice, subtracting slightly less than .5 helps give expected results even when the tool mesh is outside of the terrain.
+				UseTargetWindingThreshold -= .49;
+			}
+			else if (ToolMeshEmbedding == MeshPartition::EBooleanToolMeshEmbedding::Inside)
+			{
+				// The missing full-enclosure would have added 1 to the winding number.
+				UseTargetWindingThreshold -= 1.;
+			}
 		}
 		FMeshBoolean::FCustomInsideMeshTest TargetWindingTest(
 			[UseTargetWindingThreshold](const FVector3d& Pos, const FMeshBoolean::FCustomInsideTestContext& Context) -> bool
@@ -284,12 +306,19 @@ namespace MegaMeshBooleanLocals
 			TargetWindingTest,
 			ToolWindingTest);
 
-		// There should just be one polygroup layer for base IDs
-		ensure(CombinedMesh.Attributes()->NumPolygroupLayers() == 1
-		&& CombinedMesh.Attributes()->GetPolygroupLayer(0)->GetName() == FName(TEXT("MegaMeshBaseID")));
-		// Assign 'no base' to triangles that came from the tool mesh
-		if (FDynamicMeshPolygroupAttribute* BaseIDLayer = CombinedMesh.Attributes()->GetPolygroupLayer(0))
+		// If the tool mesh could have added its own geometry, assign 'no base' to triangles that came from the tool mesh
+		if (!FMeshBoolean::OperatesOnSingleMesh(BooleanOpMode) && 
+			// We should have a triangle mapping from the boolean operation
+			ensure(MeshBooleanOp.TrackPerTriangleSourceMesh->Num() == CombinedMesh.MaxTriangleID()) &&
+			// We should have a single MegaMeshBaseID polygroup layer on the target
+			ensure(
+				CombinedMesh.Attributes()->NumPolygroupLayers() == 1 &&
+				CombinedMesh.Attributes()->GetPolygroupLayer(0) != nullptr &&
+				CombinedMesh.Attributes()->GetPolygroupLayer(0)->GetName() == FName(TEXT("MegaMeshBaseID"))
+			)
+		)
 		{
+			FDynamicMeshPolygroupAttribute* BaseIDLayer = CombinedMesh.Attributes()->GetPolygroupLayer(0);
 			for (int32 TID : CombinedMesh.TriangleIndicesItr())
 			{
 				if ((*MeshBooleanOp.TrackPerTriangleSourceMesh)[TID] == 1)
diff --git a/Engine/Plugins/Experimental/MeshPartition/Source/MeshPartitionEditor/Public/Modifiers/MeshPartitionBooleanModifier.h b/Engine/Plugins/Experimental/MeshPartition/Source/MeshPartitionEditor/Public/Modifiers/MeshPartitionBooleanModifier.h
--- a/Engine/Plugins/Experimental/MeshPartition/Source/MeshPartitionEditor/Public/Modifiers/MeshPartitionBooleanModifier.h
+++ b/Engine/Plugins/Experimental/MeshPartition/Source/MeshPartitionEditor/Public/Modifiers/MeshPartitionBooleanModifier.h
@@ -22,11 +22,14 @@ namespace UE::MeshPartition
 UENUM()
 enum class EBooleanOperation
 {
-	/** Union of the two objects, i.e. merger of both objects  */
+	/** Merges the Tool and Mesh Partition geometry into a locally-closed surface, removing internal surfaces where they overlap. */
 	Union = 0 UMETA(DisplayName = "Union"),
 
-	/** Remove the attached mesh geometry from the MegaMesh geometry */
-	Subtract = 1 UMETA(DisplayName = "Subtract")
+	/** Subtracts the Tool mesh from the Mesh Partition geometry, removing overlapping geometry and capping the result into a locally-closed surface. */
+	Subtract = 1 UMETA(DisplayName = "Subtract"),
+
+	/** Cuts away Mesh Partition geometry inside the Tool mesh, leaving an open hole at the boundary. */
+	Trim = 2 UMETA(DisplayName = "Trim")
 };
 
 /** Method to use for determining what is 'inside' a mesh */
@@ -117,7 +120,7 @@ public:
 
 	// MeshPartition::UModifierComponent
 	UE_API virtual TArray<FBox> ComputeBounds() const override;
-	virtual bool IsContiguous() const { return true; }
+	virtual bool IsContiguous() const { return BooleanOp != EBooleanOperation::Trim; }
 	UE_API TSharedPtr<const MeshPartition::IModifierBackgroundOp> CreateBackgroundOp(const MeshPartition::EBuildType InBuildType) const;
 	UE_API virtual FGuid GetCodeVersionKey() const override;
 	UE_API virtual void GatherDependencies(MeshPartition::IDependencyInterface& Dependencies) const override;
@@ -194,7 +197,7 @@ private:
 	* This is useful for e.g. using Boolean subtractions to create caverns underneath the surface terrain, 
 	* without needing to extend the Boolean operator bounds all the way to the surface.
 	*/
-	UPROPERTY(EditAnywhere, Category = Boolean)
+	UPROPERTY(EditAnywhere, Category = Boolean, meta = (EditCondition="BooleanOp != EBooleanOperation::Trim"))
 	MeshPartition::EBooleanToolMeshEmbedding ToolMeshEmbedding = MeshPartition::EBooleanToolMeshEmbedding::Intersecting;
 
 	/* Method to use for determining the inside of the target mesh. */
diff --git a/Engine/Source/Runtime/GeometryCore/Private/Operations/MeshBoolean.cpp b/Engine/Source/Runtime/GeometryCore/Private/Operations/MeshBoolean.cpp
--- a/Engine/Source/Runtime/GeometryCore/Private/Operations/MeshBoolean.cpp
+++ b/Engine/Source/Runtime/GeometryCore/Private/Operations/MeshBoolean.cpp
@@ -265,7 +265,7 @@ bool FMeshBoolean::ComputeWithCustomInside(FCustomInsideMeshTest InsideMeshA, FC
 		return false;
 	}
 
-	bool bOpOnSingleMesh = Operation == EBooleanOp::TrimInside || Operation == EBooleanOp::TrimOutside || Operation == EBooleanOp::NewGroupInside || Operation == EBooleanOp::NewGroupOutside;
+	bool bOpOnSingleMesh = OperatesOnSingleMesh(Operation);
 
 	// cut the meshes
 	FMeshMeshCut Cut(CutMesh[0], CutMesh[1]);
diff --git a/Engine/Source/Runtime/GeometryCore/Public/Operations/MeshBoolean.h b/Engine/Source/Runtime/GeometryCore/Public/Operations/MeshBoolean.h
--- a/Engine/Source/Runtime/GeometryCore/Public/Operations/MeshBoolean.h
+++ b/Engine/Source/Runtime/GeometryCore/Public/Operations/MeshBoolean.h
@@ -141,6 +141,7 @@ public:
 	// If set, will be populated with an array of source mesh indices per TriangleID in the result mesh,
 	// indicating if that triangle was from the Meshes[0] or Meshes[1] source.
 	// Note: For non-compact meshes, invalid triangle IDs will have arbitrary indices.
+	// Note: Only filled for operations that combine both meshes (i.e., when OperatesOnSingleMesh(Operation) == false)
 	TOptional<TArray<int8>> TrackPerTriangleSourceMesh;
 
 public:
@@ -184,6 +185,13 @@ public:
 	 */
 	GEOMETRYCORE_API bool Compute();
 
+	// Helper to test if an operation will incorporate/edit both meshes, or just one
+	// Note: If the operation works on a single mesh, then TrackPerTriangleSourceMesh will not be filled
+	static bool OperatesOnSingleMesh(EBooleanOp Operation)
+	{
+		return Operation == EBooleanOp::TrimInside || Operation == EBooleanOp::TrimOutside || Operation == EBooleanOp::NewGroupInside || Operation == EBooleanOp::NewGroupOutside;
+	}
+
 	// Data passed to the FCustomInsideMeshTest::IsPointInsideFn
 	struct FCustomInsideTestContext
 	{
