[Mesh Partition]

Create per vertex uvs and per vertex normals when converted mesh partition mesh data to dynamic meshes.

#Synced-CL 51841952
#Synced-CL 51841952

diff --git a/Engine/Plugins/Experimental/MeshPartition/Source/MeshPartition/Private/MeshPartitionMeshData.cpp b/Engine/Plugins/Experimental/MeshPartition/Source/MeshPartition/Private/MeshPartitionMeshData.cpp
--- a/Engine/Plugins/Experimental/MeshPartition/Source/MeshPartition/Private/MeshPartitionMeshData.cpp
+++ b/Engine/Plugins/Experimental/MeshPartition/Source/MeshPartition/Private/MeshPartitionMeshData.cpp
@@ -499,6 +499,8 @@ void FMeshData::ConvertToDynamicMesh(Geometry::FDynamicMesh3& OutDynamicMesh) co
 	Geometry::FDynamicMeshUVOverlay* UVOverlay = OutDynamicMesh.Attributes()->PrimaryUV();
 
 	OutDynamicMesh.BeginUnsafeVerticesInsert();
+	NormalOverlay->BeginUnsafeElementsInsert();
+	UVOverlay->BeginUnsafeElementsInsert();
 	for (int VertexIndex : VertexIndicesItr())
 	{
 		Geometry::FVertexInfo VertexInfo;
@@ -507,6 +509,12 @@ void FMeshData::ConvertToDynamicMesh(Geometry::FDynamicMesh3& OutDynamicMesh) co
 
 		OutDynamicMesh.InsertVertex(VertexIndex, VertexInfo, true);
 		VertexMap.Add(VertexIndex, VertexID);
+		
+		const FVector2f UV = GetVertexUV(VertexIndex);
+		UVOverlay->InsertElement(VertexID, &UV.X, true);
+
+		const FVector3f Normal = GetVertexNormal(VertexIndex);
+		NormalOverlay->InsertElement(VertexID, &Normal.X, true);
 
 		for (const TPair<FName, TArray<float>>& WeightLayer : WeightLayers)
 		{
@@ -514,6 +522,8 @@ void FMeshData::ConvertToDynamicMesh(Geometry::FDynamicMesh3& OutDynamicMesh) co
 		}
 	}
 	OutDynamicMesh.EndUnsafeVerticesInsert();
+	NormalOverlay->EndUnsafeElementsInsert();
+	UVOverlay->EndUnsafeElementsInsert();
 		
 	OutDynamicMesh.BeginUnsafeTrianglesInsert();
 	for (int TriangleIndex : TriangleIndicesItr())
@@ -522,21 +532,10 @@ void FMeshData::ConvertToDynamicMesh(Geometry::FDynamicMesh3& OutDynamicMesh) co
 
 		const Geometry::FIndex3i ResultTriangle(VertexMap[Triangle.A], VertexMap[Triangle.B], VertexMap[Triangle.C]);
 
-		Geometry::FIndex3i UVTriangle;
-		Geometry::FIndex3i NormalTriangle;
-		for (int Index = 0; Index < 3; ++Index)
-		{
-			const FVector2f UV = GetVertexUV(Triangle[Index]);
-			UVTriangle[Index] = UVOverlay->AppendElement(UV);
-
-			const FVector3f Normal = GetVertexNormal(Triangle[Index]);
-			NormalTriangle[Index] = NormalOverlay->AppendElement(Normal);
-		}
-
 		OutDynamicMesh.InsertTriangle(TriangleIndex, ResultTriangle, 0, /* bUnsafe = */ true);
 
-		UVOverlay->SetTriangle(TriangleIndex, UVTriangle);
-		NormalOverlay->SetTriangle(TriangleIndex, NormalTriangle);
+		UVOverlay->SetTriangle(TriangleIndex, ResultTriangle);
+		NormalOverlay->SetTriangle(TriangleIndex, ResultTriangle);
 	}
 	OutDynamicMesh.EndUnsafeTrianglesInsert();
 }
