[Mesh Partition]

Added warnings and prevent creation of Mesh Partition when in non-world partitioned maps.

#Synced-CL 51758550

diff --git a/Engine/Plugins/Experimental/MeshPartition/Source/MeshPartitionEditor/Private/MeshPartitionEditorComponent.cpp b/Engine/Plugins/Experimental/MeshPartition/Source/MeshPartitionEditor/Private/MeshPartitionEditorComponent.cpp
--- a/Engine/Plugins/Experimental/MeshPartition/Source/MeshPartitionEditor/Private/MeshPartitionEditorComponent.cpp
+++ b/Engine/Plugins/Experimental/MeshPartition/Source/MeshPartitionEditor/Private/MeshPartitionEditorComponent.cpp
@@ -28,6 +28,8 @@
 #include "Engine/TextureRenderTarget2DArray.h"
 #include "Engine/World.h"
 #include "MaterialDomain.h"
+#include "Misc/UObjectToken.h"
+#include "Logging/MessageLog.h"
 
 #include "MeshPartition.h"
 #include "MeshPartitionCompiledSection.h"
@@ -197,6 +199,24 @@ void UMeshPartitionEditorComponent::OnUnregister()
 	Super::OnUnregister();
 }
 
+void UMeshPartitionEditorComponent::CheckForErrors()
+{
+	Super::CheckForErrors();
+
+	const UWorld* World = GetWorld();
+	if (!IsTemplate() && World && World->WorldType != EWorldType::EditorPreview)
+	{
+		if (World->GetWorldPartition() == nullptr)
+		{
+			FMessageLog("MapCheck").AddMessage(FTokenizedMessage::Create(EMessageSeverity::Error)
+						->AddToken(FUObjectToken::Create(this))
+						->AddToken(FTextToken::Create(
+							LOCTEXT("MapCheck_Message_MeshPartitionInNonWPWorld", "Mesh Partition requires a map with World Partition enabled.")
+						)));
+		}
+	}
+}
+
 void UMeshPartitionEditorComponent::PostRegisterMegaMeshComponents()
 {
 	// This is necessary because actor deletion will null out our CurrentModifiers before we get a chance to uninitialize
diff --git a/Engine/Plugins/Experimental/MeshPartition/Source/MeshPartitionEditor/Public/MeshPartitionEditorComponent.h b/Engine/Plugins/Experimental/MeshPartition/Source/MeshPartitionEditor/Public/MeshPartitionEditorComponent.h
--- a/Engine/Plugins/Experimental/MeshPartition/Source/MeshPartitionEditor/Public/MeshPartitionEditorComponent.h
+++ b/Engine/Plugins/Experimental/MeshPartition/Source/MeshPartitionEditor/Public/MeshPartitionEditorComponent.h
@@ -67,6 +67,7 @@ public:
 
 	// Begin ActorComponent Implementation
 	UE_API virtual void OnUnregister() override;
+	virtual void CheckForErrors() override;
 	// End ActorComponent Implementation
 
 	// Begin UMeshPartitionComponent Implementation
diff --git a/Engine/Plugins/Experimental/MeshPartition/Source/MeshPartitionModelingToolset/Private/MeshPartitionConvertTool.cpp b/Engine/Plugins/Experimental/MeshPartition/Source/MeshPartitionModelingToolset/Private/MeshPartitionConvertTool.cpp
--- a/Engine/Plugins/Experimental/MeshPartition/Source/MeshPartitionModelingToolset/Private/MeshPartitionConvertTool.cpp
+++ b/Engine/Plugins/Experimental/MeshPartition/Source/MeshPartitionModelingToolset/Private/MeshPartitionConvertTool.cpp
@@ -23,6 +23,8 @@
 
 using namespace UE::Geometry;
 
+#define LOCTEXT_NAMESPACE "MeshPartition::UConvertTool"
+
 namespace UE::MeshPartition
 {
 USingleSelectionMeshEditingTool* MeshPartition::UConvertToolBuilder::CreateNewTool(const FToolBuilderState& InSceneState) const
@@ -58,10 +60,22 @@ void UConvertTool::Setup()
 	CreateProperties = NewObject<MeshPartition::UCreateProperties>(this);
 	AddToolPropertySource(CreateProperties);
 	CreateProperties->RestoreProperties(this);
+
+	if (!TargetWorld.IsValid() || !TargetWorld->GetWorldPartition())
+	{
+		GetToolManager()->DisplayMessage(
+			LOCTEXT("MeshPartition_WorldPartitionRequired", "Mesh Partition requires World Partition to be enabled."),
+			EToolMessageLevel::UserError);
+	}
 }
 
 bool UConvertTool::CanAccept() const
 {
+	if (!TargetWorld.IsValid() || !TargetWorld->GetWorldPartition())
+	{
+		return false;
+	}
+
 	return true;
 }
 
@@ -197,3 +211,5 @@ TArray<TUniquePtr<FDynamicMesh3>> MeshPartition::UConvertTool::SplitMesh(const F
 	return SplitMeshes;
 }
 } // namespace UE::MeshPartition
+
+#undef LOCTEXT_NAMESPACE
\ No newline at end of file
diff --git a/Engine/Plugins/Experimental/MeshPartition/Source/MeshPartitionModelingToolset/Private/MeshPartitionCreateMeshTool.cpp b/Engine/Plugins/Experimental/MeshPartition/Source/MeshPartitionModelingToolset/Private/MeshPartitionCreateMeshTool.cpp
--- a/Engine/Plugins/Experimental/MeshPartition/Source/MeshPartitionModelingToolset/Private/MeshPartitionCreateMeshTool.cpp
+++ b/Engine/Plugins/Experimental/MeshPartition/Source/MeshPartitionModelingToolset/Private/MeshPartitionCreateMeshTool.cpp
@@ -203,11 +203,6 @@ UInteractiveTool* UCreateMeshToolBuilder::BuildTool(const FToolBuilderState& Sce
 void UCreateMeshTool::SetWorld(UWorld* InWorld)
 {
 	World = InWorld;
-
-	if (CreateRectangleProperties)
-	{
-		CreateRectangleProperties->SetWorldPartitionAvailable( World && World->GetWorldPartition() != nullptr);
-	}
 }
 
 UCreateMeshTool::UCreateMeshTool(const FObjectInitializer&)
@@ -219,6 +214,11 @@ UCreateMeshTool::UCreateMeshTool(const FObjectInitializer&)
 
 bool UCreateMeshTool::CanAccept() const
 {
+	if (!World || !World->GetWorldPartition())
+	{
+		return false;
+	}
+
 	return CurrentState == EState::AdjustingSettings;
 }
 
@@ -348,6 +348,7 @@ void UCreateMeshTool::Setup()
 
 	UpdatePreview();
 	SetState(EState::PlacingPrimitive);
+	UpdateDisplayMessage();
 }
 
 void UCreateMeshTool::SetState(EState NewState)
@@ -744,6 +745,16 @@ void UCreateMeshTool::UpdatePreview()
 
 }
 
+void UCreateMeshTool::UpdateDisplayMessage() const
+{
+	if (!World || !World->GetWorldPartition())
+	{
+		GetToolManager()->DisplayMessage(
+			LOCTEXT("MeshPartition_WorldPartitionRequired", "Mesh Partition requires World Partition to be enabled."),
+			EToolMessageLevel::UserError);
+	}
+}
+
 void UCreateMeshTool::UpdateTargetSurface()
 {
 	if (PlacementProperties->TargetSurface == MeshPartition::ECreateMeshPlacementType::AtOrigin)
@@ -1031,25 +1042,26 @@ void UCreateRectangleMeshTool::Shutdown(EToolShutdownType InShutdownType)
 
 
 
-void UCreateRectangleMeshTool::CheckMeshResolution() const
+void UCreateRectangleMeshTool::UpdateDisplayMessage() const
 {
-	UE_LOG(LogTemp, Display, TEXT("CheckMeshResolution"));
+	Super::UpdateDisplayMessage();
 
 	constexpr int32 Threshold = 4096 * 4096 * 2; // Warn above 4k resolution.
 	constexpr int32 ThresholdSaveAndUnload = Threshold * 4 * 4; // Warn above 16k resolution when using "Save and Unload".
 
-	const bool bWorldPartitionAvailable = World && World->GetWorldPartition() != nullptr;
-	const bool bSaveAndUnloadEnabled = bWorldPartitionAvailable && CreateRectangleProperties->bSaveAndUnload;
+	const bool bSaveAndUnloadEnabled = CreateRectangleProperties->bSaveAndUnload;
 	const int32 NumTriangles = CreateRectangleProperties->MeshResolution.X * CreateRectangleProperties->MeshResolution.Y * 2;
 
 	FText Message;
+	EToolMessageLevel MessageLevel = EToolMessageLevel::Internal;
 
 	if (!bSaveAndUnloadEnabled && NumTriangles > Threshold)
 	{
 		Message = FText::Format(
 			LOCTEXT("LargeResolution", "Imported mesh will consist of {0} triangles.{1}"),
 			FText::AsNumber(NumTriangles),
-			bWorldPartitionAvailable ? LOCTEXT("EnableSaveAndUnload", " Consider enabling \"Save and Unload\" to limit memory consumption.") : FText());
+			LOCTEXT("EnableSaveAndUnload", " Consider enabling \"Save and Unload\" to limit memory consumption."));
+		MessageLevel = EToolMessageLevel::UserNotification;
 	}
 
 	if (bSaveAndUnloadEnabled && NumTriangles > ThresholdSaveAndUnload)
@@ -1057,6 +1069,12 @@ void UCreateRectangleMeshTool::CheckMeshResolution() const
 		Message = FText::Format(
 			LOCTEXT("LargeResolutionSaveAndUnload", "Imported mesh will consist of {0} triangles."),
 			FText::AsNumber(NumTriangles));
+		MessageLevel = EToolMessageLevel::UserNotification;
+	}
+	
+	if (!Message.IsEmpty())
+	{
+		GetToolManager()->DisplayMessage(Message, MessageLevel);
 	}
 }
 
diff --git a/Engine/Plugins/Experimental/MeshPartition/Source/MeshPartitionModelingToolset/Private/MeshPartitionHeightmapImportTool.cpp b/Engine/Plugins/Experimental/MeshPartition/Source/MeshPartitionModelingToolset/Private/MeshPartitionHeightmapImportTool.cpp
--- a/Engine/Plugins/Experimental/MeshPartition/Source/MeshPartitionModelingToolset/Private/MeshPartitionHeightmapImportTool.cpp
+++ b/Engine/Plugins/Experimental/MeshPartition/Source/MeshPartitionModelingToolset/Private/MeshPartitionHeightmapImportTool.cpp
@@ -77,9 +77,9 @@ void UHeightmapImportTool::Setup()
 
 	Properties->RestoreProperties(this);
 
-	Properties->WatchProperty(Properties->MeshResolution, [this](const FInt32Point&) { CheckMeshResolution(); });
-	Properties->WatchProperty(Properties->bSaveAndUnload, [this](const bool&) { CheckMeshResolution(); });
-	CheckMeshResolution();
+	Properties->WatchProperty(Properties->MeshResolution, [this](const FInt32Point&) { UpdateDisplayMessage(); });
+	Properties->WatchProperty(Properties->bSaveAndUnload, [this](const bool&) { UpdateDisplayMessage(); });
+	UpdateDisplayMessage();
 
 	AddToolPropertySource(Properties);
 }
@@ -149,6 +149,11 @@ void UHeightmapImportTool::Render(IToolsContextRenderAPI* RenderAPI)
 
 bool UHeightmapImportTool::CanAccept() const
 {
+	if (!World || !World->GetWorldPartition())
+	{
+		return false;
+	}
+
 	const FInt32Point SectionsResolution = Properties->GetSectionsResolution();
 	const FInt32Point LocationVolumesResolution = Properties->GetLocationVolumesResolution();
 
@@ -158,16 +163,20 @@ bool UHeightmapImportTool::CanAccept() const
 		&& SectionsResolution.X > 0 && SectionsResolution.Y > 0
 		&& LocationVolumesResolution.X >= 0 && LocationVolumesResolution.Y >= 0;
 }
-
-void UHeightmapImportTool::CheckMeshResolution() const
+void UHeightmapImportTool::UpdateDisplayMessage() const
 {
-	UE_LOG(LogTemp, Display, TEXT("CheckMeshResolution"));
+	if (!World || !World->GetWorldPartition())
+	{
+		GetToolManager()->DisplayMessage(
+			LOCTEXT("MeshPartition_WorldPartitionRequired", "Mesh Partition requires World Partition to be enabled."),
+			EToolMessageLevel::UserError);
+		return;
+	}
 
 	constexpr int32 Threshold = 4096 * 4096 * 2; // Warn above 4k resolution.
 	constexpr int32 ThresholdSaveAndUnload = Threshold * 4 * 4; // Warn above 16k resolution when using "Save and Unload".
 
-	const bool bWorldPartitionAvailable = World && World->GetWorldPartition() != nullptr;
-	const bool bSaveAndUnloadEnabled = bWorldPartitionAvailable && Properties->bSaveAndUnload;
+	const bool bSaveAndUnloadEnabled = Properties->bSaveAndUnload;
 	const int32 NumTriangles = Properties->MeshResolution.X * Properties->MeshResolution.Y * 2;
 
 	FText Message;
@@ -177,7 +186,7 @@ void UHeightmapImportTool::CheckMeshResolution() const
 		Message = FText::Format(
 			LOCTEXT("LargeResolution", "Imported mesh will consist of {0} triangles.{1}"),
 			FText::AsNumber(NumTriangles),
-			bWorldPartitionAvailable ? LOCTEXT("EnableSaveAndUnload", " Consider enabling \"Save and Unload\" to limit memory consumption.") : FText());
+			LOCTEXT("EnableSaveAndUnload", " Consider enabling \"Save and Unload\" to limit memory consumption."));
 	}
 
 	if (bSaveAndUnloadEnabled && NumTriangles > ThresholdSaveAndUnload)
diff --git a/Engine/Plugins/Experimental/MeshPartition/Source/MeshPartitionModelingToolset/Public/MeshPartitionCreateMeshTool.h b/Engine/Plugins/Experimental/MeshPartition/Source/MeshPartitionModelingToolset/Public/MeshPartitionCreateMeshTool.h
--- a/Engine/Plugins/Experimental/MeshPartition/Source/MeshPartitionModelingToolset/Public/MeshPartitionCreateMeshTool.h
+++ b/Engine/Plugins/Experimental/MeshPartition/Source/MeshPartitionModelingToolset/Public/MeshPartitionCreateMeshTool.h
@@ -191,7 +191,7 @@ public:
 	*
 	* This is only available when World Partition is enabled.
 	*/
-	UPROPERTY(EditAnywhere, Category=Sections, meta=(EditCondition="bWorldPartitionAvailable", HideEditConditionToggle))
+	UPROPERTY(EditAnywhere, Category=Sections)
 	bool bSaveAndUnload = false;
 
 	
@@ -199,11 +199,6 @@ public:
 private:
 	friend class UCreateRectangleMeshTool;
 	friend class UCreateMeshTool;
-	
-	void SetWorldPartitionAvailable(bool bInWorldPartitionAvailable)
-	{
-		bWorldPartitionAvailable = bInWorldPartitionAvailable;
-	}
 
 	// Resolution in terms of sections, but not all sections will necessarily contain the same number of quads 
 	FInt32Point ComputeMeshResolutionInSections() const;
@@ -211,9 +206,6 @@ private:
 	// compute mesh resolution in terms of quads
 	FInt32Point ComputeMeshResolutionInQuads() const;
 
-	UPROPERTY(meta=(TransientToolProperty))
-	bool bWorldPartitionAvailable;
-
 	UPROPERTY(meta=(TransientToolProperty))
 	FInt32Point CachedMeshResolution = FInt32Point(-1, -1);;
 };
@@ -229,7 +221,7 @@ public:
 	*
 	* This is only available when World Partition is enabled.
 	*/
-	UPROPERTY(EditAnywhere, Category=LocationVolumes, meta=(EditCondition="bWorldPartitionAvailable", HideEditConditionToggle, DisplayName="Create Volumes"))
+	UPROPERTY(EditAnywhere, Category=LocationVolumes, meta=(DisplayName="Create Volumes"))
 	bool bCreateLocationVolumes = false;
 
 	/**
@@ -237,23 +229,15 @@ public:
 	* 
 	* This is only available when World Partition is enabled.
 	*/
-	UPROPERTY(EditAnywhere, Category=LocationVolumes, meta=(EditCondition="bCreateLocationVolumes", DisplayName="Volume Resolution"))
+	UPROPERTY(EditAnywhere, Category=LocationVolumes, meta=(DisplayName="Volume Resolution"))
 	FInt32Point LocationVolumesResolution = FInt32Point(2, 2);
 
 		
 private:
 	friend class UCreateRectangleMeshTool;
 	friend class UCreateMeshTool;
-	
-	void SetWorldPartitionAvailable(bool bInWorldPartitionAvailable)
-	{
-		bWorldPartitionAvailable = bInWorldPartitionAvailable;
-	}
 
 	FInt32Point GetLocationVolumesResolution() const;
-
-	UPROPERTY(meta=(TransientToolProperty))
-	bool bWorldPartitionAvailable;
 };
 // -------------------------------------------------------------------------------------------------------------------------
 
@@ -315,6 +299,8 @@ protected:
 	
 	virtual void UpdatePreview();
 
+	virtual void UpdateDisplayMessage() const;
+
 
 	virtual void GenerateAsset() {};
 
@@ -415,7 +401,7 @@ protected:
 
 	void RenderSections(IToolsContextRenderAPI* RenderAPI) const;
 
-	void CheckMeshResolution() const;
+	virtual void UpdateDisplayMessage() const override;
 	FMegaMeshRectangleGeneratorParams GetGeneratorParams() const ;
 
 	TArray<FSectionInfo> SectionInfos;
diff --git a/Engine/Plugins/Experimental/MeshPartition/Source/MeshPartitionModelingToolset/Public/MeshPartitionHeightmapImportTool.h b/Engine/Plugins/Experimental/MeshPartition/Source/MeshPartitionModelingToolset/Public/MeshPartitionHeightmapImportTool.h
--- a/Engine/Plugins/Experimental/MeshPartition/Source/MeshPartitionModelingToolset/Public/MeshPartitionHeightmapImportTool.h
+++ b/Engine/Plugins/Experimental/MeshPartition/Source/MeshPartitionModelingToolset/Public/MeshPartitionHeightmapImportTool.h
@@ -139,7 +139,7 @@ private:
 	UPROPERTY()
 	TObjectPtr<UWorld> World;
 
-	void CheckMeshResolution() const;
+	void UpdateDisplayMessage() const;
 	FHeightmapImportParams GetParams() const;
 };
 } // namespace UE::MeshPartition
diff --git a/Engine/Plugins/Experimental/MeshTerrainMode/Source/MeshTerrainMode/Private/MeshTerrainModeToolkit.cpp b/Engine/Plugins/Experimental/MeshTerrainMode/Source/MeshTerrainMode/Private/MeshTerrainModeToolkit.cpp
--- a/Engine/Plugins/Experimental/MeshTerrainMode/Source/MeshTerrainMode/Private/MeshTerrainModeToolkit.cpp
+++ b/Engine/Plugins/Experimental/MeshTerrainMode/Source/MeshTerrainMode/Private/MeshTerrainModeToolkit.cpp
@@ -9,6 +9,7 @@
 #include "MeshTerrainModeSettings.h"
 #include "MeshTerrainModeStyle.h"
 #include "MeshTerrainMode.h"
+#include "Engine/World.h"
 
 #include "Modules/ModuleManager.h"
 #include "IDetailsView.h"
@@ -1595,6 +1596,16 @@ void FMeshTerrainModeToolkit::ShowRealtimeAndModeWarnings(bool bShowRealtimeWarn
 	{
 		WarningText = LOCTEXT("ModelingModeToolkitRealtimeWarning", "Realtime Mode is required for Modeling Tools to work correctly. Please enable Realtime Mode in the Viewport Options or with the Ctrl+r hotkey.");
 	}
+	else if (UEdMode* EdMode = GetScriptableEditorMode().Get())
+	{
+		if (UWorld* World = EdMode->GetWorld())
+		{
+			if (!World->GetWorldPartition())
+			{
+				WarningText = LOCTEXT("MeshTerrainModeWorldPartitionWarning", "Mesh Terrain Mode tools require World Partition to be enabled.");
+			}
+		}
+	}
 	if (!WarningText.IdenticalTo(ActiveWarning))
 	{
 		ActiveWarning = WarningText;
