[Mesh Partition]

Swapped the default uv unwrap method to the slow path temporarily until some adjustments can be made to the fast path.

#Synced-CL 52063422

diff --git a/Engine/Plugins/Experimental/MeshPartition/Source/MeshPartitionEditor/Private/MeshPartitionChannelCollection.cpp b/Engine/Plugins/Experimental/MeshPartition/Source/MeshPartitionEditor/Private/MeshPartitionChannelCollection.cpp
--- a/Engine/Plugins/Experimental/MeshPartition/Source/MeshPartitionEditor/Private/MeshPartitionChannelCollection.cpp
+++ b/Engine/Plugins/Experimental/MeshPartition/Source/MeshPartitionEditor/Private/MeshPartitionChannelCollection.cpp
@@ -40,7 +40,7 @@
 
 static TAutoConsoleVariable<int32> CVarMeshPartition_UnwrapUVsMethod(
 	TEXT("MeshPartition.Channels.UnwrapUVsMethod"),
-	1,
+	2,
 	TEXT("Sets the method used to automatically unwrap UVs. 1: Fast unwrap (with artifacts); 2: Reference box unwrap (slow). Default Value: 1"),
 	FConsoleVariableDelegate::CreateLambda([](IConsoleVariable*)
 	{
