use normal-project for height smooth, since the move direction is already normalized

#Synced-CL 52079281

diff --git a/Engine/Plugins/Experimental/MeshPartition/Source/MeshPartitionModelingToolset/Private/Sculpting/HeightBrushes.cpp b/Engine/Plugins/Experimental/MeshPartition/Source/MeshPartitionModelingToolset/Private/Sculpting/HeightBrushes.cpp
--- a/Engine/Plugins/Experimental/MeshPartition/Source/MeshPartitionModelingToolset/Private/Sculpting/HeightBrushes.cpp
+++ b/Engine/Plugins/Experimental/MeshPartition/Source/MeshPartitionModelingToolset/Private/Sculpting/HeightBrushes.cpp
@@ -157,7 +157,7 @@ namespace HeightBrushesLocals
 				SmoothedPos = UE::Geometry::FMeshWeights::UniformCentroid(*Mesh, VID, GetVertexFn);
 			}
 			FVector3d TargetLocation = UE::Geometry::Lerp(OrigPos, SmoothedPos, Falloff * Stamp.Power);
-			return OrigPos + (TargetLocation - OrigPos).ProjectOnTo(GetMoveDirection(OrigPos));
+			return OrigPos + (TargetLocation - OrigPos).ProjectOnToNormal(GetMoveDirection(OrigPos));
 		});
 	}
 
