Refer to Mesh Layers Modifier more consistently as Brush Modifier in user tooling
+ make the Modifier Factory use the component DisplayName to name actor/component when available

#Synced-CL 52079281

diff --git a/Engine/Plugins/Experimental/MeshPartition/Source/MeshPartitionEditor/Private/MeshPartitionModifierActor.cpp b/Engine/Plugins/Experimental/MeshPartition/Source/MeshPartitionEditor/Private/MeshPartitionModifierActor.cpp
--- a/Engine/Plugins/Experimental/MeshPartition/Source/MeshPartitionEditor/Private/MeshPartitionModifierActor.cpp
+++ b/Engine/Plugins/Experimental/MeshPartition/Source/MeshPartitionEditor/Private/MeshPartitionModifierActor.cpp
@@ -24,7 +24,19 @@ namespace ModifierActorLocals
 		}
 		
 		TargetActor->Modify();
-		FName NewComponentName = *FComponentEditorUtils::GenerateValidVariableName(ComponentType, TargetActor);
+		FName NewComponentName;
+		if (ComponentType->HasMetaData(TEXT("DisplayName")))
+		{
+			NewComponentName = *ComponentType->GetMetaData(TEXT("DisplayName"));
+			while (!FComponentEditorUtils::IsComponentNameAvailable(NewComponentName.ToString(), TargetActor))
+			{
+				NewComponentName.SetNumber(NewComponentName.GetNumber() + 1);
+			}
+		}
+		else
+		{
+			NewComponentName = *FComponentEditorUtils::GenerateValidVariableName(ComponentType, TargetActor);
+		}
 		USceneComponent* Component = NewObject<USceneComponent>(TargetActor, ComponentType, NewComponentName, RF_Transactional);
 		Component->SetupAttachment(TargetActor->GetRootComponent());
 		Component->OnComponentCreated();
@@ -138,6 +150,18 @@ AActor* UModifierActorFactory::SpawnActor(UObject* InAsset, ULevel* InLevel, con
 	}
 	return NewActor;
 }
+
+FString UModifierActorFactory::GetDefaultActorLabel(UObject* Asset) const
+{
+	if (const UClass* AssetClass = Cast<UClass>(Asset))
+	{
+		if (AssetClass->HasMetaData(TEXT("DisplayName")))
+		{
+			return AssetClass->GetMetaData(TEXT("DisplayName"));
+		}
+	}
+	return Super::GetDefaultActorLabel(Asset);
+}
 } // namespace UE::MeshPartition
 
 #undef LOCTEXT_NAMESPACE
diff --git a/Engine/Plugins/Experimental/MeshPartition/Source/MeshPartitionEditor/Public/MeshPartitionModifierActor.h b/Engine/Plugins/Experimental/MeshPartition/Source/MeshPartitionEditor/Public/MeshPartitionModifierActor.h
--- a/Engine/Plugins/Experimental/MeshPartition/Source/MeshPartitionEditor/Public/MeshPartitionModifierActor.h
+++ b/Engine/Plugins/Experimental/MeshPartition/Source/MeshPartitionEditor/Public/MeshPartitionModifierActor.h
@@ -43,6 +43,7 @@ class UModifierActorFactory : public UActorFactory
 	//~ Begin UActorFactory Interface
 	UE_API virtual AActor* SpawnActor(UObject* InAsset, ULevel* InLevel, const FTransform& InTransform, const FActorSpawnParameters& InSpawnParams) override;
 	UE_API virtual bool CanCreateActorFrom(const FAssetData& AssetData, FText& OutErrorMsg) override;
+	UE_API virtual FString GetDefaultActorLabel(UObject* Asset) const override;
 	//~ End UActorFactory Interface
 };
 } // namespace UE::MeshPartition
diff --git a/Engine/Plugins/Experimental/MeshPartition/Source/MeshPartitionEditor/Public/Modifiers/MeshPartitionProjectSculptLayersModifier.h b/Engine/Plugins/Experimental/MeshPartition/Source/MeshPartitionEditor/Public/Modifiers/MeshPartitionProjectSculptLayersModifier.h
--- a/Engine/Plugins/Experimental/MeshPartition/Source/MeshPartitionEditor/Public/Modifiers/MeshPartitionProjectSculptLayersModifier.h
+++ b/Engine/Plugins/Experimental/MeshPartition/Source/MeshPartitionEditor/Public/Modifiers/MeshPartitionProjectSculptLayersModifier.h
@@ -82,7 +82,7 @@ struct FSculptLayerModifierWeightAttributeEntry
 /**
 * Modifier to manage layers of custom mesh detail, and project them onto a partitioned mesh
 */
-UCLASS(MinimalAPI, meta = (BlueprintSpawnableComponent, DisplayName = "MeshLayersModifier", MegaMeshClassVersion = "1"))
+UCLASS(MinimalAPI, meta = (BlueprintSpawnableComponent, DisplayName = "BrushModifier", MegaMeshClassVersion = "1"))
 class UProjectMeshLayersModifier : public MeshPartition::UEditableModifierBase
 	, public IMeshSculptLayersManager
 	, public ICodeReusableModifier
diff --git a/Engine/Plugins/Experimental/MeshTerrainMode/Source/MeshTerrainMode/Private/Submodes/ModifiersSubmode.cpp b/Engine/Plugins/Experimental/MeshTerrainMode/Source/MeshTerrainMode/Private/Submodes/ModifiersSubmode.cpp
--- a/Engine/Plugins/Experimental/MeshTerrainMode/Source/MeshTerrainMode/Private/Submodes/ModifiersSubmode.cpp
+++ b/Engine/Plugins/Experimental/MeshTerrainMode/Source/MeshTerrainMode/Private/Submodes/ModifiersSubmode.cpp
@@ -44,7 +44,7 @@ FModifiersSubmode::FModifiersSubmode(const TSharedPtr<FModeToolkit>& InToolkit)
 				CreatePlaceableItem(MeshPartition::UMeshProjectModifier::StaticClass(), LOCTEXT("MeshProjectModifier", "Mesh")),
 				CreatePlaceableItem(MeshPartition::UTexturePatchModifier::StaticClass(), LOCTEXT("TexturePatchModifier", "Texture")),
 				CreatePlaceableItem(MeshPartition::USplineModifier::StaticClass(), LOCTEXT("SplineModifier", "Spline")),
-				CreatePlaceableItem(MeshPartition::UProjectMeshLayersModifier::StaticClass(), LOCTEXT("SculptLayersModifier", "Sculpt")),
+				CreatePlaceableItem(MeshPartition::UProjectMeshLayersModifier::StaticClass(), LOCTEXT("SculptLayersModifier", "Brush")),
 				CreatePlaceableItem(MeshPartition::ULatticeModifier::StaticClass(), LOCTEXT("LatticeModifier", "Lattice")),
 				CreatePlaceableItem(MeshPartition::UNoiseModifier::StaticClass(), LOCTEXT("NoiseModifier", "Noise")),
 				CreatePlaceableItem(MeshPartition::UPatchModifier::StaticClass(), LOCTEXT("PatchModifier", "Patch")),
