Remove out-of-date restriction on height brushes not supporting spacing mode, and make the height sculpt brush support spacing mode

#Synced-CL 52079281

diff --git a/Engine/Plugins/Experimental/MeshPartition/Source/MeshPartitionModelingToolset/Private/Sculpting/HeightBrushes.h b/Engine/Plugins/Experimental/MeshPartition/Source/MeshPartitionModelingToolset/Private/Sculpting/HeightBrushes.h
--- a/Engine/Plugins/Experimental/MeshPartition/Source/MeshPartitionModelingToolset/Private/Sculpting/HeightBrushes.h
+++ b/Engine/Plugins/Experimental/MeshPartition/Source/MeshPartitionModelingToolset/Private/Sculpting/HeightBrushes.h
@@ -138,20 +138,6 @@ public:
 		return true;
 	}
 
-	virtual bool SupportsStrokeType(EMeshSculptStrokeType StrokeType) const override
-	{
-		switch (StrokeType)
-		{
-		case EMeshSculptStrokeType::Airbrush:
-		case EMeshSculptStrokeType::Dots:
-			return true;
-		// TODO: Currently there is an issue with using spacing mode (see UE-307860)
-		case EMeshSculptStrokeType::Spacing:
-		default:
-			return false;
-		}
-	}
-
 	virtual EReferencePlaneType GetReferencePlaneType() const override
 	{
 		return IsSphereBrush() ? EReferencePlaneType::WorkPlane : EReferencePlaneType::None;
@@ -180,6 +166,19 @@ public:
 		: FHeightBrushOpBase(HeightSculptPropertiesIn)
 	{}
 
+	virtual bool SupportsStrokeType(EMeshSculptStrokeType StrokeType) const override
+	{
+		switch (StrokeType)
+		{
+		case EMeshSculptStrokeType::Airbrush:
+		case EMeshSculptStrokeType::Dots:
+		case EMeshSculptStrokeType::Spacing:
+			return true;
+		default:
+			return false;
+		}
+	}
+
 	virtual bool UsesAlpha() const override { return true; }
 
 	virtual void ApplyStamp(const FDynamicMesh3* SrcMesh, const FSculptBrushStamp& Stamp,
