WP grid-aligned compiled sections: split mesh by WP grid cell size

Split a modifier group's compiled section mesh by the World Partition grid cell size so each ACompiledSection aligns to WP streaming boundaries. Prevents large sections from spanning multiple WP cells.

New:

* bAlignToWPGrid opt-in toggle on FCompiledSectionBuildVariant
* Grid cell size auto-detected from the variant's FWPActorPropertiesTransformer ? RuntimeGrid ? GetGridCellSize(); fallback to main grid if named grid not found
* PrepareCompiledSections unified entry point: dispatches to grid (N sections) or single (1 section) path
* Grid-split reuse: all N sibling sections tracked and reused together (ReuseCompiledSectionActorDescGuids)
* GridCellCoord (FIntVector) stored in FCompiledSectionBuildInfo for per-section identity
* Resolved grid cell size hashed into BuildVariantHash at builder level

PIE / SeparateWorldBuilder support:

* WorldUpdater creates N grid-split placeholders using ComputeBaseBounds()
* SeparateWorldBuilder builds grid-split sections on-the-fly with SplitAndCacheGridBuild
* BuildVariantHash check added to WorldUpdater reuse loop (was missing)
* RuntimeCellTransformer strips stale sections on variant hash mismatch

Bug fixes:

* Fix bounds mismatch between WorldUpdater placeholders and BuildSections (use ComputeBaseBounds() in both paths)
* Fix iterator invalidation crash in FSeparateWorldBuilder::Tick
* Null-guard GetLoadedLevel() and GetWorld() for PIE streaming-out scenarios
* Fix FDependencyHash hiding base-class templated operator+=

#Synced-CL 52205931
#Synced-CL 52205931

diff --git a/Engine/Plugins/Experimental/MeshPartition/Source/MeshPartition/Private/MeshPartitionCompiledSection.cpp b/Engine/Plugins/Experimental/MeshPartition/Source/MeshPartition/Private/MeshPartitionCompiledSection.cpp
--- a/Engine/Plugins/Experimental/MeshPartition/Source/MeshPartition/Private/MeshPartitionCompiledSection.cpp
+++ b/Engine/Plugins/Experimental/MeshPartition/Source/MeshPartition/Private/MeshPartitionCompiledSection.cpp
@@ -498,7 +498,8 @@ bool FCompiledSectionBuildInfo::operator== (const FCompiledSectionBuildInfo& Oth
 		(ClassHash == Other.ClassHash) &&
 		(BuildVariantHash == Other.BuildVariantHash) &&
 		(MegaMeshGUID == Other.MegaMeshGUID) &&
-		(MegaMeshPath == Other.MegaMeshPath);
+		(MegaMeshPath == Other.MegaMeshPath) &&
+		(GridCellCoord == Other.GridCellCoord);
 }
 
 bool FCompiledSectionBuildInfo::SetMegaMeshDefinition(const UMeshPartitionDefinition* Definition)
@@ -563,15 +564,29 @@ void FCompiledSectionBuildInfo::Serialize(FArchive& Ar)
 		}
 #endif // WITH_EDITORONLY_DATA
 	}
+
+	if (Version >= MeshPartition::FCustomVersion::AddedGridCellCoordToBuildInfo)
+	{
+		Ar << GridCellCoord.X;
+		Ar << GridCellCoord.Y;
+		Ar << GridCellCoord.Z;
+		Ar << ResolvedGridCellSize;
+	}
 }
 
 bool FCompiledSectionBuildInfo::TargetsSameCompiledSectionAs(const FCompiledSectionBuildInfo& Other) const
 {
+	// If either side has an invalid grid coordinate (non-grid section), skip the coordinate check
+	// so non-wp-aligned sections remain compatible with the wp aligned pipeline.
+	const bool bGridCoordMatches = (!HasValidGridCellCoord() || !Other.HasValidGridCellCoord()) || (GridCellCoord == Other.GridCellCoord);
+
 	return (BuildVariantName == Other.BuildVariantName) &&
 		(MegaMeshGUID == Other.MegaMeshGUID) &&
 		(MegaMeshDefinitionPath == Other.MegaMeshDefinitionPath) &&
 		(BaseModifierPaths == Other.BaseModifierPaths) &&
-		((!MegaMeshPath.IsValid()) || (MegaMeshPath == Other.MegaMeshPath)); // If the path is not valid, it's likely just an old compiled section, default to accept
+		// If the path is not valid, it's likely just an old compiled section, default to accept
+		((!MegaMeshPath.IsValid()) || (MegaMeshPath == Other.MegaMeshPath)) &&
+		bGridCoordMatches;
 }
 
 #if WITH_EDITOR
diff --git a/Engine/Plugins/Experimental/MeshPartition/Source/MeshPartition/Private/MeshPartitionDefinition.cpp b/Engine/Plugins/Experimental/MeshPartition/Source/MeshPartition/Private/MeshPartitionDefinition.cpp
--- a/Engine/Plugins/Experimental/MeshPartition/Source/MeshPartition/Private/MeshPartitionDefinition.cpp
+++ b/Engine/Plugins/Experimental/MeshPartition/Source/MeshPartition/Private/MeshPartitionDefinition.cpp
@@ -132,6 +132,10 @@ void FCommonBuildVariant::GatherDependencies(MeshPartition::IDependencyInterface
 void FCompiledSectionBuildVariant::GatherDependencies(MeshPartition::IDependencyInterface& InDependencies) const
 {
 	FCommonBuildVariant::GatherDependencies(InDependencies);
+
+	InDependencies += bSplitSectionsToMatchWorldPartitionRuntimeGrid;
+	// Note: resolved grid cell size is hashed at builder level (not here)
+	// because it requires world context to resolve from the WP runtime partition.
 }
 #endif // WITH_EDITOR
 
diff --git a/Engine/Plugins/Experimental/MeshPartition/Source/MeshPartition/Public/MeshPartitionCompiledSection.h b/Engine/Plugins/Experimental/MeshPartition/Source/MeshPartition/Public/MeshPartitionCompiledSection.h
--- a/Engine/Plugins/Experimental/MeshPartition/Source/MeshPartition/Public/MeshPartitionCompiledSection.h
+++ b/Engine/Plugins/Experimental/MeshPartition/Source/MeshPartition/Public/MeshPartitionCompiledSection.h
@@ -112,6 +112,22 @@ public:
 	UPROPERTY(VisibleAnywhere, Category = "MegaMesh")
 	FSoftObjectPath MegaMeshPath;
 
+	// Sentinel for sections not produced by a grid-aligned split.
+	static inline const FIntVector InvalidGridCellCoord = FIntVector(MIN_int32, MIN_int32, MIN_int32);
+
+	// Absolute world-space grid coordinate of this section's cell.
+	UPROPERTY(VisibleAnywhere, Category = "MegaMesh", meta = (NoResetToDefault))
+	FIntVector GridCellCoord = InvalidGridCellCoord;
+
+	// Resolved WP grid cell size used to split this section. Set by the builder or WorldUpdater
+	// at creation time so consumers (e.g. SeparateWorldBuilder) don't need to re-resolve from the pipeline.
+	// Not Transient: CPF_Transient causes StaticDuplicateObject to skip the property, which breaks
+	// PIE world duplication (editor -> PIE actor copy). Regular UPROPERTY survives both duplication
+	// and serialization; the default of 0 is harmless for sections loaded from disk.
+	// Not part of identity or equality comparisons.
+	UPROPERTY(VisibleAnywhere, Category = "MegaMesh")
+	uint32 ResolvedGridCellSize = 0;
+
 private:
 	// transient cache of the mega mesh definition pointer (converted from the MegaMeshDefinitionPath)
 	mutable const UMeshPartitionDefinition* MegaMeshDefinition;
@@ -122,6 +138,7 @@ public:
 	MESHPARTITION_API bool TargetsSameCompiledSectionAs(const FCompiledSectionBuildInfo& Other) const;
 	MESHPARTITION_API bool SetMegaMeshDefinition(const UMeshPartitionDefinition* Definition);
 	MESHPARTITION_API const UMeshPartitionDefinition* GetMegaMeshDefinition() const;
+	bool HasValidGridCellCoord() const { return GridCellCoord != InvalidGridCellCoord; }
 };
 
 struct FCompiledSectionDescriptor
diff --git a/Engine/Plugins/Experimental/MeshPartition/Source/MeshPartition/Public/MeshPartitionDefinition.h b/Engine/Plugins/Experimental/MeshPartition/Source/MeshPartition/Public/MeshPartitionDefinition.h
--- a/Engine/Plugins/Experimental/MeshPartition/Source/MeshPartition/Public/MeshPartitionDefinition.h
+++ b/Engine/Plugins/Experimental/MeshPartition/Source/MeshPartition/Public/MeshPartitionDefinition.h
@@ -75,6 +75,9 @@ public:
 	UPROPERTY(EditAnywhere, Category="Build Variant")
 	FName Name = NAME_Default;
 
+	UPROPERTY(EditAnywhere, Category = "WorldPartition")
+	bool bSplitSectionsToMatchWorldPartitionRuntimeGrid = false;
+
 protected:
 #if WITH_EDITOR
 	friend class UMeshPartitionDefinition;
diff --git a/Engine/Plugins/Experimental/MeshPartition/Source/MeshPartition/Public/MeshPartitionModule.h b/Engine/Plugins/Experimental/MeshPartition/Source/MeshPartition/Public/MeshPartitionModule.h
--- a/Engine/Plugins/Experimental/MeshPartition/Source/MeshPartition/Public/MeshPartitionModule.h
+++ b/Engine/Plugins/Experimental/MeshPartition/Source/MeshPartition/Public/MeshPartitionModule.h
@@ -31,6 +31,8 @@ struct FCustomVersion
 		BuildVariantCleanup = 4,
 		// Added debug checksums to the package and class hashes in the compiled section build infos
 		AddedChecksumsToBuildInfo = 5,
+		// Added GridCellCoord (FIntVector) and ResolvedGridCellSize to FCompiledSectionBuildInfo for WP-grid-aligned sections
+		AddedGridCellCoordToBuildInfo = 6,
 		// Add more versions as needed
 
 		LatestVersionPlusOne,
diff --git a/Engine/Plugins/Experimental/MeshPartition/Source/MeshPartition/Public/MeshPartitionTransformer.h b/Engine/Plugins/Experimental/MeshPartition/Source/MeshPartition/Public/MeshPartitionTransformer.h
--- a/Engine/Plugins/Experimental/MeshPartition/Source/MeshPartition/Public/MeshPartitionTransformer.h
+++ b/Engine/Plugins/Experimental/MeshPartition/Source/MeshPartition/Public/MeshPartitionTransformer.h
@@ -30,6 +30,12 @@ struct FTransformerUnit
 	bool bShouldRecomputeTangents = false;
 };
 
+/** Construct a transformer unit from a section actor and its mesh data. */
+inline FTransformerUnit MakeTransformerUnit(AActor* InSection, TSharedPtr<const FMeshData> InMesh, bool bInRecomputeNormals = true, bool bInRecomputeTangents = true)
+{
+	return FTransformerUnit{InSection, MoveTemp(InMesh), bInRecomputeNormals, bInRecomputeTangents};
+}
+
 struct FTransformerContext
 {
 	UE_API ~FTransformerContext();
@@ -51,7 +57,7 @@ struct FTransformer
 {
 	GENERATED_BODY()
 
-	UE_API virtual ~FTransformer() = default;
+	virtual ~FTransformer() = default;
 
 	/**
 	 * After intantiating the transformer. Initializes it with the definition in use.
diff --git a/Engine/Plugins/Experimental/MeshPartition/Source/MeshPartitionEditor/Private/MeshPartitionChannelCollection.cpp b/Engine/Plugins/Experimental/MeshPartition/Source/MeshPartitionEditor/Private/MeshPartitionChannelCollection.cpp
--- a/Engine/Plugins/Experimental/MeshPartition/Source/MeshPartitionEditor/Private/MeshPartitionChannelCollection.cpp
+++ b/Engine/Plugins/Experimental/MeshPartition/Source/MeshPartitionEditor/Private/MeshPartitionChannelCollection.cpp
@@ -7,6 +7,8 @@
 
 #include "Tasks/Task.h"
 #include "Async/ParallelFor.h"
+
+#include "MaterialCache/MaterialCache.h"
 #include "MeshPartitionChannelRasterizationShaders.h"
 #include "MeshPartitionEditorComponent.h"
 #include "MeshPartitionModifierComponent.h"
@@ -36,7 +38,25 @@
 #include "UObject/ConstructorHelpers.h"
 
 #include "UObject/Object.h"
+#include "MeshPartitionCompiledSection.h"
+
+namespace UE::MeshPartition
+{
+
+void ApplyChannels(ACompiledSection* InSection, const FSectionChannels& InChannels, const float InMaterialCacheTexelSize)
+{
+	InSection->SetChannelTexture(InChannels.Texture.Get());
+	InSection->SetChannelTexcoordDesc(InChannels.TexcoordMetrics);
+	InSection->SetChannelTable(InChannels.Table);
+
+	if (InChannels.Domain.AreaUV > 0)
+	{
+		InSection->SetMaterialCacheTileCount(GetMaterialCacheTileCount(InChannels.Domain.Area3D / InChannels.Domain.AreaUV, InMaterialCacheTexelSize));
+		InSection->RecreateMaterialCacheTextures();
+	}
+}
 
+} // namespace UE::MeshPartition
 
 static TAutoConsoleVariable<int32> CVarMeshPartition_UnwrapUVsMethod(
 	TEXT("MeshPartition.Channels.UnwrapUVsMethod"),
diff --git a/Engine/Plugins/Experimental/MeshPartition/Source/MeshPartitionEditor/Private/MeshPartitionEditorComponent.cpp b/Engine/Plugins/Experimental/MeshPartition/Source/MeshPartitionEditor/Private/MeshPartitionEditorComponent.cpp
--- a/Engine/Plugins/Experimental/MeshPartition/Source/MeshPartitionEditor/Private/MeshPartitionEditorComponent.cpp
+++ b/Engine/Plugins/Experimental/MeshPartition/Source/MeshPartitionEditor/Private/MeshPartitionEditorComponent.cpp
@@ -453,57 +453,107 @@ void UMeshPartitionEditorComponent::BuildMegaMeshPreviewSections(TConstArrayView
 	}
 }
 
-MeshPartition::ACompiledSection* UMeshPartitionEditorComponent::PrepareMegaMeshCompiledSection(const MeshPartition::FCompiledSectionBuildInfo& InBuildInfo, const MeshPartition::FCompiledSectionBuildVariant& InBuildVariant, MeshPartition::ACompiledSection* InReuseCompiledSection, bool bInUseStaticMobility) const
+TArray<MeshPartition::ACompiledSection*> UMeshPartitionEditorComponent::PrepareCompiledSections(const MeshPartition::FCompiledSectionBuildInfo&		InBuildInfo,
+																								const MeshPartition::FCompiledSectionBuildVariant&	InBuildVariant,
+																								MeshPartition::FPrepareCompiledSectionsParams&		InOutParams) const
 {
-	TRACE_CPUPROFILER_EVENT_SCOPE(UMeshPartitionEditorComponent::PrepareMegaMeshCompiledSection);
-	
+	TRACE_CPUPROFILER_EVENT_SCOPE(UMeshPartitionEditorComponent::PrepareCompiledSections);
+
 	UWorld* World = GetWorld();
-	const AMeshPartition* MegaMesh = GetOwner<const AMeshPartition>();
+	const AMeshPartition* MeshPartition = GetOwner<const AMeshPartition>();
 
 	if (!ensure(World != nullptr))
 	{
-		UE_LOG(LogMegaMeshEditor, Error, TEXT("Cannot get World."));
-		return nullptr;
-	}
-	
-	if (!ensure(MegaMesh != nullptr))
-	{
-		UE_LOG(LogMegaMeshEditor, Error, TEXT("Cannot get MegaMesh."));
-		return nullptr;
+		UE_LOG(LogMegaMeshEditor, Error, TEXT("PrepareCompiledSections: Cannot get World."));
+		return {};
 	}
-	
-	MeshPartition::ACompiledSection* CompiledSection = InReuseCompiledSection;
-	if (CompiledSection == nullptr)
-	{
-		FString BaseName = TEXT("CompiledSection_") + InBuildInfo.BuildVariantName.ToString();
 
-		CompiledSection = GetWorld()->SpawnActor<MeshPartition::ACompiledSection>(MeshPartition::ACompiledSection::StaticClass(), FTransform::Identity);
-		CompiledSection->SetActorLabel(BaseName);
+	if (!ensure(MeshPartition != nullptr))
+	{
+		UE_LOG(LogMegaMeshEditor, Error, TEXT("PrepareCompiledSections: Cannot get MeshPartition."));
+		return {};
 	}
 
-	CompiledSection->SetParent(Cast<AMeshPartition>(GetOwner()));
+	// Resolve the base material once for all sections.
+	UMaterialInterface* BaseMaterial = const_cast<UMaterialInterface*>(GetDefinitionMaterial());
 
-	if (CompiledSection->GetRootComponent())
+	if (!IsValid(BaseMaterial))
 	{
-		CompiledSection->GetRootComponent()->SetMobility(bInUseStaticMobility ? EComponentMobility::Static : EComponentMobility::Stationary);
+		BaseMaterial = UMaterial::GetDefaultMaterial(MD_Surface);
 	}
 
-	// Record BuildInfo on the CompiledSection
-	CompiledSection->SetBuildInfo(InBuildInfo);
+	// Spawns or reuses a single ACompiledSection with material setup and build info.
+	auto PrepareSection = [&](const MeshPartition::FCompiledSectionBuildInfo&	InSectionBuildInfo,
+							  MeshPartition::ACompiledSection*					InReuseSection = nullptr,
+							  bool												bInStaticMobility = true) -> MeshPartition::ACompiledSection*
+	{
+		MeshPartition::ACompiledSection* Section = InReuseSection;
+
+		if (Section == nullptr)
+		{
+			FString BaseName = TEXT("CompiledSection_") + InSectionBuildInfo.BuildVariantName.ToString();
+			Section = World->SpawnActor<MeshPartition::ACompiledSection>(MeshPartition::ACompiledSection::StaticClass(), FTransform::Identity);
+
+			if (Section == nullptr)
+			{
+				UE_LOG(LogMegaMeshEditor, Error, TEXT("Failed to spawn ACompiledSection for variant '%s'"), *InSectionBuildInfo.BuildVariantName.ToString());
+				return nullptr;
+			}
+
+			Section->SetActorLabel(BaseName);
+		}
+
+		Section->SetParent(Cast<AMeshPartition>(GetOwner()));
 
-	// Setup the material instance constant:
-	UMaterialInterface* BaseMegaMeshMaterial = const_cast<UMaterialInterface*>(GetDefinitionMaterial());
-	if (!IsValid(BaseMegaMeshMaterial))
+		if (Section->GetRootComponent())
+		{
+			Section->GetRootComponent()->SetMobility(bInStaticMobility ? EComponentMobility::Static : EComponentMobility::Stationary);
+		}
+
+		Section->SetBuildInfo(InSectionBuildInfo);
+
+		UMaterialInstanceConstant* MIC = NewObject<UMaterialInstanceConstant>(Section->GetPackage(), MakeUniqueObjectName(Section, UMaterialInstanceConstant::StaticClass(), TEXT("CompiledSectionMIC")));
+		MIC->SetParentEditorOnly(BaseMaterial);
+		Section->SetMaterialInstance(MIC);
+
+		return Section;
+	};
+
+	TArray<MeshPartition::ACompiledSection*> Sections;
+	const bool bIsGridSplitValid = InBuildVariant.bSplitSectionsToMatchWorldPartitionRuntimeGrid && (InOutParams.GridCellSize > 0.0);
+
+	if (InBuildVariant.bSplitSectionsToMatchWorldPartitionRuntimeGrid && !bIsGridSplitValid && !InOutParams.ReuseSection)
 	{
-		BaseMegaMeshMaterial = UMaterial::GetDefaultMaterial(MD_Surface);
+		UE_LOG(LogMegaMeshEditor, Warning, TEXT("bSplitSectionsToMatchWorldPartitionRuntimeGrid is true but GridCellSize could not be resolved — falling back to single-section build."));
 	}
-	UMaterialInstanceConstant* MIC = NewObject<UMaterialInstanceConstant>(CompiledSection->GetPackage(), MakeUniqueObjectName(CompiledSection, UMaterialInstanceConstant::StaticClass(), TEXT("CompiledSectionMIC")));
 
-	MIC->SetParentEditorOnly(BaseMegaMeshMaterial);
+	if (bIsGridSplitValid)
+	{
+		InOutParams.OutCellMeshes = MeshPartition::GridHelpers::BuildGridCellMeshes(InOutParams.FullMesh, InOutParams.GridCellSize);
+
+		Sections.Reserve(InOutParams.OutCellMeshes.Num());
 
-	CompiledSection->SetMaterialInstance(MIC);
+		for (TPair<FIntVector, MeshPartition::FMeshData>& CellEntry : InOutParams.OutCellMeshes)
+		{
+			MeshPartition::FCompiledSectionBuildInfo CellBuildInfo = InBuildInfo;
+			CellBuildInfo.GridCellCoord = CellEntry.Key;
+			CellBuildInfo.ResolvedGridCellSize = static_cast<uint32>(InOutParams.GridCellSize);
 
-	return CompiledSection;
+			if (MeshPartition::ACompiledSection* Section = PrepareSection(CellBuildInfo))
+			{
+				Sections.Add(Section);
+			}
+		}
+	}
+	else
+	{
+		if (MeshPartition::ACompiledSection* Section = PrepareSection(InBuildInfo, InOutParams.ReuseSection, InOutParams.bUseStaticMobility))
+		{
+			Sections.Add(Section);
+		}
+	}
+
+	return Sections;
 }
 
 void UMeshPartitionEditorComponent::ExecuteTransformers(MeshPartition::UPreviewMeshComponent* InPreviewMeshComponent, const MeshPartition::UMeshPartitionDefinition* InDefinition, const MeshPartition::FCommonBuildVariant& InBuildVariant)
@@ -524,7 +574,7 @@ void UMeshPartitionEditorComponent::ExecuteTransformers(MeshPartition::UPreviewM
 	}
 }
 
-TUniquePtr<MeshPartition::FTransformerContext> UMeshPartitionEditorComponent::LaunchTransformers(MeshPartition::FTransformerUnit&& InTransformerUnit, const MeshPartition::UMeshPartitionDefinition* InDefinition, const MeshPartition::FCommonBuildVariant& InBuildVariant)
+TUniquePtr<MeshPartition::FTransformerContext> UMeshPartitionEditorComponent::LaunchTransformers(TArray<MeshPartition::FTransformerUnit>&& InTransformerUnits, const MeshPartition::UMeshPartitionDefinition* InDefinition, const MeshPartition::FCommonBuildVariant& InBuildVariant)
 {
 	TUniquePtr<MeshPartition::FTransformerContext> TransformersContext;
 
@@ -548,7 +598,7 @@ TUniquePtr<MeshPartition::FTransformerContext> UMeshPartitionEditorComponent::La
 		CopiedTransformer.GetMutable().Initialize(InDefinition, InBuildVariant);
 	}
 
-	MeshPartition::FTransformerUnit& TransformerUnit = TransformersContext->TransformerUnits.Emplace_GetRef(MoveTemp(InTransformerUnit));
+	TransformersContext->TransformerUnits = MoveTemp(InTransformerUnits);
 
 	for (int32 Index = 0; Index < TransformersContext->Transformers.Num(); ++Index)
 	{
@@ -567,7 +617,7 @@ TUniquePtr<MeshPartition::FTransformerContext> UMeshPartitionEditorComponent::La
 			  	{
 					return;
 			  	}
-			  	
+
 				Transformer.Get().Execute(*TransformersContextPtr);
 			}, UE::Tasks::Prerequisites(Prerequsites));
 	}
@@ -581,6 +631,7 @@ void UMeshPartitionEditorComponent::BuildMegaMeshCompiledSectionTextures(MeshPar
 	FChannelMap ChannelMap;
 	float ChannelTexelSize = 100;
 	float MaterialCacheTexelSize = 1;
+	
 	if (UMeshPartitionDefinition* Definition = GetMegaMeshDefinition())
 	{
 		ChannelMap = Definition->GetChannelMap();
@@ -591,19 +642,7 @@ void UMeshPartitionEditorComponent::BuildMegaMeshCompiledSectionTextures(MeshPar
 	constexpr bool bDownloadToAsset = true;
 	MeshPartition::FSectionChannels SectionChannelTexture = FChannelTextureRenderer::BuildTextureForSection(InOutBuiltMesh, InCompiledSection, bDownloadToAsset, ChannelMap, ChannelTexelSize).GetResult();
 
-	InCompiledSection->SetChannelTexcoordDesc(SectionChannelTexture.TexcoordMetrics); 
-	InCompiledSection->SetChannelTable(SectionChannelTexture.Table);
-	InCompiledSection->SetChannelTexture(SectionChannelTexture.Texture.Get());
-	
-	if (SectionChannelTexture.Domain.AreaUV > 0)
-	{
-		InCompiledSection->SetMaterialCacheTileCount(GetMaterialCacheTileCount(
-			SectionChannelTexture.Domain.Area3D / SectionChannelTexture.Domain.AreaUV,
-			MaterialCacheTexelSize
-		));
-		
-		InCompiledSection->RecreateMaterialCacheTextures();
-	}
+	MeshPartition::ApplyChannels(InCompiledSection, SectionChannelTexture, MaterialCacheTexelSize);
 
 	// Coming out from the ChannelCollection section generation the ChannelTexture is ready to be used and parented to the CompiledSection
 	if (!SectionChannelTexture.Texture)
@@ -696,6 +735,8 @@ void UMeshPartitionEditorComponent::PostBuildSectionMesh(AActor* InSection, cons
 	// The PostBuildSectionMesh calls are not currently asynchronous
 	ensure(IsInGameThread());
 
+	const FBox SectionBounds(InBuiltMesh.GetBounds());
+
 	for (TWeakObjectPtr<MeshPartition::UModifierComponent> Modifier : InModifiers)
 	{
 		if (!ensure(Modifier.IsValid()))
@@ -703,6 +744,12 @@ void UMeshPartitionEditorComponent::PostBuildSectionMesh(AActor* InSection, cons
 			continue;
 		}
 
+		// When sections are grid-split, a modifier may not overlap this cell — skip it.
+		if (SectionBounds.IsValid && !Modifier->IntersectsAnyBounds({ SectionBounds }))
+		{
+			continue;
+		}
+
 		Modifier->PostBuildSectionMesh(InSection, InBuiltMesh);
 		Modifier->UpdateLastAppliedBounds();
 	}
@@ -715,6 +762,8 @@ void UMeshPartitionEditorComponent::PostProcessSection(AActor* InSection, TConst
 	// The PostProcessSection calls are not currently asynchronous
 	ensure(IsInGameThread());
 
+	const FBox SectionBounds = InSection->GetComponentsBoundingBox();
+
 	for (TWeakObjectPtr<MeshPartition::UModifierComponent> Modifier : InModifiers)
 	{
 		if (!ensure(Modifier.IsValid()))
@@ -722,11 +771,66 @@ void UMeshPartitionEditorComponent::PostProcessSection(AActor* InSection, TConst
 			continue;
 		}
 
+		// When sections are grid-split, a modifier may not overlap this cell — skip it.
+		if (SectionBounds.IsValid && !Modifier->IntersectsAnyBounds({ SectionBounds }))
+		{
+			continue;
+		}
+
 		Modifier->PostProcessSection(InSection);
 		Modifier->UpdateLastAppliedBounds();
 	}
 }
 
+void UMeshPartitionEditorComponent::ApplyBuiltMeshToSection(MeshPartition::ACompiledSection*					InSection,
+															TSharedPtr<const MeshPartition::FMeshData>			InMesh,
+															const MeshPartition::FSectionChannels&				InChannels,
+															const MeshPartition::FModifierGroup&				InModifierGroup)
+{
+	TRACE_CPUPROFILER_EVENT_SCOPE(UMeshPartitionEditorComponent::ApplyBuiltMeshToSection);
+
+	const MeshPartition::FCompiledSectionBuildInfo& BuildInfo = InSection->GetBuildInfo();
+
+	// Resolve definition and build variant from the section's build info.
+	MeshPartition::UMeshPartitionDefinition* Definition = TSoftObjectPtr<MeshPartition::UMeshPartitionDefinition>(FSoftObjectPath(BuildInfo.MegaMeshDefinitionPath)).Get();
+
+	if (!Definition)
+	{
+		Definition = MeshPartition::UMeshPartitionDefinition::StaticClass()->GetDefaultObject<MeshPartition::UMeshPartitionDefinition>();
+	}
+
+	const MeshPartition::FCompiledSectionBuildVariant& BuildVariant = Definition->GetCompiledSectionBuildVariantByName(BuildInfo.BuildVariantName);
+
+	// Prepare section (material, build info, reuse existing placeholder actor)
+	MeshPartition::FPrepareCompiledSectionsParams Params{.FullMesh = *InMesh, .ReuseSection = InSection, .bUseStaticMobility = true};
+	TArray<MeshPartition::ACompiledSection*> PreparedSections = PrepareCompiledSections(BuildInfo, BuildVariant, Params);
+
+	if (!ensure(!PreparedSections.IsEmpty()))
+	{
+		UE_LOG(LogMegaMeshEditor, Error, TEXT("ApplyBuiltMeshToSection: PrepareCompiledSections returned no sections — skipping."));
+		return;
+	}
+
+	ensureMsgf(PreparedSections[0] == InSection, TEXT("PrepareCompiledSections did not reuse the provided section actor."));
+
+	MeshPartition::ApplyChannels(InSection, InChannels, Definition->GetMaterialCacheTexelSize());
+
+	TArray<TWeakObjectPtr<MeshPartition::UModifierComponent>> ModifierPtrs = InModifierGroup.AllResolvedModifierPtrs();
+	PostBuildSectionMesh(InSection, *InMesh, ModifierPtrs);
+
+	TArray<MeshPartition::FTransformerUnit> Units;
+	Units.Add(MeshPartition::MakeTransformerUnit(InSection, InMesh));
+
+	TUniquePtr<MeshPartition::FTransformerContext> TransformerContext = LaunchTransformers(MoveTemp(Units), Definition, BuildVariant);
+
+	if (TransformerContext.IsValid() && TransformerContext->JoinTask.IsValid() && !TransformerContext->JoinTask.IsCompleted())
+	{
+		MeshPartition::WaitOnGameThread(*TransformerContext);
+	}
+
+	PostProcessSection(InSection, ModifierPtrs);
+}
+
 void UMeshPartitionEditorComponent::OnModifierChanged(MeshPartition::UModifierComponent* InChangedModifier, TConstArrayView<FBox> InBoundingBoxes, const MeshPartition::EChangeType InChangeType)
 {
 	if (!IsInteractiveModeReady())
@@ -1631,24 +1735,28 @@ void UMeshPartitionEditorComponent::FinalizePreviewSectionBuilds()
 			PreviewSection->SetPreviewMesh(ProcessedMesh.ToSharedRef(), UCollisionProfile::BlockAll_ProfileName, /* bCanEverAffectNavigation = */ false);
 			PreviewSection->SetBuildPerfStats(TaskHandle.GetTask()->GetBuildPerfStats());
 
-			PreviewSection->SetChannelTexture(SectionChannel->Texture.Get());
-			PreviewSection->SetChannelTexcoordDesc(SectionChannel->TexcoordMetrics);
-			PreviewSection->SetChannelTable(SectionChannel->Table);
-			if (SectionChannel->Mesh.IsSet())
+			if (SectionChannel.IsValid())
 			{
-				PreviewSection->SetChannelGenerationMesh(SectionChannel->Mesh.GetValue());
-			}
+				PreviewSection->SetChannelTexture(SectionChannel->Texture.Get());
+				PreviewSection->SetChannelTexcoordDesc(SectionChannel->TexcoordMetrics);
+				PreviewSection->SetChannelTable(SectionChannel->Table);
 
-			if (UMeshPartitionDefinition* Definition = GetMegaMeshDefinition())
-			{
-				if (SectionChannel->Domain.AreaUV > 0)
+				if (SectionChannel->Mesh.IsSet())
 				{
-					PreviewSection->SetMaterialCacheTileCount(GetMaterialCacheTileCount(
-						SectionChannel->Domain.Area3D / SectionChannel->Domain.AreaUV,
-						Definition->GetMaterialCacheTexelSize()
-					));
+					PreviewSection->SetChannelGenerationMesh(SectionChannel->Mesh.GetValue());
+				}
 
-					PreviewSection->RecreateMaterialCacheTextures();
+				if (UMeshPartitionDefinition* Definition = GetMegaMeshDefinition())
+				{
+					if (SectionChannel->Domain.AreaUV > 0)
+					{
+						PreviewSection->SetMaterialCacheTileCount(GetMaterialCacheTileCount(
+							SectionChannel->Domain.Area3D / SectionChannel->Domain.AreaUV,
+							Definition->GetMaterialCacheTexelSize()
+						));
+
+						PreviewSection->RecreateMaterialCacheTextures();
+					}
 				}
 			}
 
@@ -1667,14 +1775,10 @@ void UMeshPartitionEditorComponent::FinalizePreviewSectionBuilds()
 
 			if (!bIsTransformerPipelinePaused)
 			{
-				MeshPartition::FTransformerUnit InitialUnit;
-
-				InitialUnit.Section = PreviewSection;
-				InitialUnit.MeshData = ProcessedMesh;
-				InitialUnit.bShouldRecomputeNormals = false;
-				InitialUnit.bShouldRecomputeTangents = ShouldRecomputeTangents();
+				TArray<MeshPartition::FTransformerUnit> InitialUnits;
+				InitialUnits.Add(MeshPartition::MakeTransformerUnit(PreviewSection, ProcessedMesh, /*bRecomputeNormals=*/ false, ShouldRecomputeTangents()));
 
-				TUniquePtr<MeshPartition::FTransformerContext> TransformerContext = LaunchTransformers(MoveTemp(InitialUnit), GetMegaMeshDefinition(), GetPreviewBuildVariant());
+				TUniquePtr<MeshPartition::FTransformerContext> TransformerContext = LaunchTransformers(MoveTemp(InitialUnits), GetMegaMeshDefinition(), GetPreviewBuildVariant());
 
 				if (TransformerContext.IsValid())
 				{
@@ -1772,20 +1876,16 @@ void UMeshPartitionEditorComponent::SetPauseTransformerPipeline(const bool bInSh
 			{
 				const bool bIsPaused = false;
 
-				MeshPartition::FTransformerUnit InitialUnit;
+				TArray<MeshPartition::FTransformerUnit> InitialUnits;
+				InitialUnits.Add(MeshPartition::MakeTransformerUnit(PreviewSection, PreviewSection->GetPreviewMesh(), /*bRecomputeNormals=*/ false, ShouldRecomputeTangents()));
 
-				InitialUnit.Section = PreviewSection;
-				InitialUnit.MeshData = PreviewSection->GetPreviewMesh();
-				InitialUnit.bShouldRecomputeNormals = false;
-				InitialUnit.bShouldRecomputeTangents = ShouldRecomputeTangents();
-
-				TUniquePtr<MeshPartition::FTransformerContext> TransformerContext = LaunchTransformers(MoveTemp(InitialUnit), GetMegaMeshDefinition(), GetPreviewBuildVariant());
+				TUniquePtr<MeshPartition::FTransformerContext> TransformerContext = LaunchTransformers(MoveTemp(InitialUnits), GetMegaMeshDefinition(), GetPreviewBuildVariant());
 
 				if (TransformerContext.IsValid())
 				{
 					PreviewSectionTransformerContexts.Emplace(MoveTemp(TransformerContext));
 				}
-				
+
 				PreviewSection->SetIsTransformerPipelinePaused(bIsPaused);
 			}
 
diff --git a/Engine/Plugins/Experimental/MeshPartition/Source/MeshPartitionEditor/Private/MeshPartitionEditorUtils.cpp b/Engine/Plugins/Experimental/MeshPartition/Source/MeshPartitionEditor/Private/MeshPartitionEditorUtils.cpp
--- a/Engine/Plugins/Experimental/MeshPartition/Source/MeshPartitionEditor/Private/MeshPartitionEditorUtils.cpp
+++ b/Engine/Plugins/Experimental/MeshPartition/Source/MeshPartitionEditor/Private/MeshPartitionEditorUtils.cpp
@@ -168,62 +168,48 @@ namespace UE::MeshPartition::EditorUtils
 		return (BuildInfo.ClassHash == ComputedClassHash);
 	}
 
-	const FCompiledSectionDescriptor* FindMatchingCompiledSectionFromPackageHashes(
+	TArray<FGuid> FindMatchingCompiledSectionsFromPackageHashes(
 		TArrayView<const FCompiledSectionDescriptor> CompiledSectionDescriptorsToSearch,
 		const MeshPartition::FCompiledSectionBuildInfo& BuildInfo,
 		IAssetRegistry& AssetRegistry)
 	{
-		TArray<const FCompiledSectionDescriptor*> MatchingCompiledSectionActorDescGUIDs;
-		if (CompiledSectionDescriptorsToSearch.Num() > 0)
+		TArray<FGuid> MatchingGuids;
+
+		for (const FCompiledSectionDescriptor& Desc : CompiledSectionDescriptorsToSearch)
 		{
-			for (const FCompiledSectionDescriptor& Desc : CompiledSectionDescriptorsToSearch)
+			if (Desc.Info.TargetsSameCompiledSectionAs(BuildInfo) &&			// ensures we are targeting the same compiled section
+				Desc.Info.BuildVariantHash == BuildInfo.BuildVariantHash &&		// catches changes to build variant
+				Desc.Info.ModifierSetHash == BuildInfo.ModifierSetHash)			// catches additions or removals from the set of modifiers (TODO: DOES NOT WORK IN LEVEL INSTANCES)
 			{
-				if (Desc.Info.TargetsSameCompiledSectionAs(BuildInfo) &&			// ensures we are targeting the same compiled section	
-					Desc.Info.BuildVariantHash == BuildInfo.BuildVariantHash &&		// catches changes to build variant
-					Desc.Info.ModifierSetHash == BuildInfo.ModifierSetHash)			// catches additions or removals from the set of modifiers (TODO: DOES NOT WORK IN LEVEL INSTANCES)
+				if (PackageHashIsUpToDate(Desc.Info, AssetRegistry) &&			// all used packages are unmodified on disk
+					ClassHashIsUpToDate(Desc.Info))								// all used classes have the same MegaMeshClassVersion
 				{
-					if (PackageHashIsUpToDate(Desc.Info, AssetRegistry))			// all used packages are unmodified on disk
-					{
-						if (ClassHashIsUpToDate(Desc.Info))							// all used classes have the same MegaMeshClassVersion
-						{
-							// found a match
-							return &Desc;
-						}
-					}
+					MatchingGuids.Add(Desc.ActorDescGuid);
 				}
 			}
 		}
-		return nullptr;
+		return MatchingGuids;
 	}
 
-	FGuid FindMatchingCompiledSectionFromModifierHash(
+	TArray<FGuid> FindMatchingCompiledSectionsFromModifierHash(
 		TArrayView<const FCompiledSectionDescriptor> CompiledSectionDescriptorsToSearch,
 		const MeshPartition::FCompiledSectionBuildInfo& BuildInfo)
 	{
-		// check if we have a matching compiled section
-		if (CompiledSectionDescriptorsToSearch.Num() > 0)
-		{
-			const FCompiledSectionDescriptor* MatchingExistingCompiledSectionDescriptor =
-				Algo::FindByPredicate(CompiledSectionDescriptorsToSearch, [&BuildInfo](const FCompiledSectionDescriptor& Desc)
-					{
-						check(Desc.Info.MegaMeshGUID == BuildInfo.MegaMeshGUID);				// we should have filtered out mismatches by this point
-
-						bool bSame =
-							Desc.Info.ModifiersHash == BuildInfo.ModifiersHash &&				// catches change to any modifier hash
-							Desc.Info.BuildVariantHash == BuildInfo.BuildVariantHash &&			// catches changes to build variant
-							Desc.Info.TargetsSameCompiledSectionAs(BuildInfo);					// ensures we are targetting the same compiled section
-							// Desc.Info.ModifierSetHash == BuildInfo.ModifierSetHash			// not necessary: is a subset of changes covered by ModifiersHash
-						return bSame;
-					});
-
-			if (MatchingExistingCompiledSectionDescriptor != nullptr)
+		TArray<FGuid> MatchingGuids;
+
+		for (const FCompiledSectionDescriptor& Desc : CompiledSectionDescriptorsToSearch)
+		{
+			check(Desc.Info.MegaMeshGUID == BuildInfo.MegaMeshGUID);				// we should have filtered out mismatches by this point
+
+			if (Desc.Info.ModifiersHash == BuildInfo.ModifiersHash &&				// catches change to any modifier hash
+				Desc.Info.BuildVariantHash == BuildInfo.BuildVariantHash &&			// catches changes to build variant
+				Desc.Info.TargetsSameCompiledSectionAs(BuildInfo))					// ensures we are targetting the same compiled section
+				// Desc.Info.ModifierSetHash == BuildInfo.ModifierSetHash			// not necessary: is a subset of changes covered by ModifiersHash
 			{
-				// we can re-use the existing compiled section
-				return MatchingExistingCompiledSectionDescriptor->ActorDescGuid;
+				MatchingGuids.Add(Desc.ActorDescGuid);
 			}
 		}
-
-		return FGuid();
+		return MatchingGuids;
 	}
 
 	bool ValidateObjectsAreUnloaded(TConstArrayView<TWeakObjectPtr<AActor>> InActorsToValidate)
diff --git a/Engine/Plugins/Experimental/MeshPartition/Source/MeshPartitionEditor/Private/MeshPartitionMeshBuilder.cpp b/Engine/Plugins/Experimental/MeshPartition/Source/MeshPartitionEditor/Private/MeshPartitionMeshBuilder.cpp
--- a/Engine/Plugins/Experimental/MeshPartition/Source/MeshPartitionEditor/Private/MeshPartitionMeshBuilder.cpp
+++ b/Engine/Plugins/Experimental/MeshPartition/Source/MeshPartitionEditor/Private/MeshPartitionMeshBuilder.cpp
@@ -153,6 +153,62 @@ namespace UE::MeshPartition::BuildStats
 }
 #endif
 
+namespace UE::MeshPartition::GridHelpers
+{
+	FGridDimensions ComputeGridDimensions(const FBox& InBounds, double InCellSize)
+	{
+		check(InCellSize > 0.0);
+
+		const FVector SnappedMin(FMath::FloorToDouble(InBounds.Min.X / InCellSize) * InCellSize,
+								 FMath::FloorToDouble(InBounds.Min.Y / InCellSize) * InCellSize,
+								 FMath::FloorToDouble(InBounds.Min.Z / InCellSize) * InCellSize);
+
+		const FIntVector OriginCoord(FMath::FloorToInt32(InBounds.Min.X / InCellSize),
+									 FMath::FloorToInt32(InBounds.Min.Y / InCellSize),
+									 FMath::FloorToInt32(InBounds.Min.Z / InCellSize));
+
+		const FBox SnappedBounds(SnappedMin, InBounds.Max);
+		const FVector Extents = SnappedBounds.GetExtent();
+
+		const int32 NumX = FMath::Max(1, FMath::DivideAndRoundUp(Extents.X * 2.0, InCellSize));
+		const int32 NumY = FMath::Max(1, FMath::DivideAndRoundUp(Extents.Y * 2.0, InCellSize));
+		const int32 NumZ = FMath::Max(1, FMath::DivideAndRoundUp(Extents.Z * 2.0, InCellSize));
+
+		const int64 TotalCells64 = static_cast<int64>(NumX) * static_cast<int64>(NumY) * static_cast<int64>(NumZ);
+		check(TotalCells64 > 0 && TotalCells64 <= MAX_int32);
+
+		return { SnappedMin, OriginCoord, { NumX, NumY, NumZ }, static_cast<int32>(TotalCells64) };
+	}
+
+	TMap<FIntVector, MeshPartition::FMeshData> BuildGridCellMeshes(const MeshPartition::FMeshData& InMesh, double InCellSize)
+	{
+		const FGridDimensions Grid = ComputeGridDimensions(FBox(InMesh.GetBounds()), InCellSize);
+
+		constexpr bool bFilterEmptyMeshes = false;
+		TArray<MeshPartition::FMeshData> FlatCells = BuildHelpers::BuildSections(InMesh, InCellSize, bFilterEmptyMeshes);
+
+		TMap<FIntVector, MeshPartition::FMeshData> Result;
+		Result.Reserve(Grid.TotalCells);
+
+		for (int32 FlatIndex = 0; FlatIndex < FlatCells.Num(); ++FlatIndex)
+		{
+			if (FlatCells[FlatIndex].VertexCount() == 0)
+			{
+				continue;
+			}
+
+			const int32 LocalX = FlatIndex % Grid.CellNumber.X;
+			const int32 LocalY = (FlatIndex / Grid.CellNumber.X) % Grid.CellNumber.Y;
+			const int32 LocalZ = FlatIndex / (Grid.CellNumber.X * Grid.CellNumber.Y);
+			const FIntVector Coord = Grid.OriginCoord + FIntVector(LocalX, LocalY, LocalZ);
+
+			Result.Add(Coord, MoveTemp(FlatCells[FlatIndex]));
+		}
+
+		return Result;
+	}
+}
+
 namespace UE::MeshPartition::BuildHelpers
 {
 	namespace Private
@@ -383,7 +439,7 @@ namespace UE::MeshPartition::BuildHelpers
 			Hasher.Update(Name.GetData(), Name.Len() * sizeof(TCHAR));
 		};
 
-		static FGuid VersionKey(TEXT("03c548f6-a408-4725-961c-020947ede321"));
+		static FGuid VersionKey(TEXT("fc71e455-e6cf-46d7-aa1f-7f4c31fc4573"));
 		UpdateHashPODType(Hasher, VersionKey);
 
 		// Hash the current mesh data type version into the key to ensure swapping mesh types hits different DDC entries.
@@ -456,26 +512,7 @@ namespace UE::MeshPartition::BuildHelpers
 		return Descriptors;
 	}
 
-	int32 GetTotalNumberOfSections(const FBox& InBounds, const double InSectionSize, int32& OutNumberOfSectionsX, int32& OutNumberOfSectionsY, int32& OutNumberOfSectionsZ)
-	{
-		if (InSectionSize <= 0)
-		{
-			OutNumberOfSectionsX = 1;
-			OutNumberOfSectionsY = 1;
-			OutNumberOfSectionsZ = 1;
-			return 1;
-		}
-
-		const UE::Geometry::FAxisAlignedBox3d MeshBounds(InBounds);
-
-		OutNumberOfSectionsX = FMath::Max(1, FMath::DivideAndRoundUp(MeshBounds.Extents().X * 2., InSectionSize));
-		OutNumberOfSectionsY = FMath::Max(1, FMath::DivideAndRoundUp(MeshBounds.Extents().Y * 2., InSectionSize));
-		OutNumberOfSectionsZ = FMath::Max(1, FMath::DivideAndRoundUp(MeshBounds.Extents().Z * 2., InSectionSize));
-		
-		return OutNumberOfSectionsX * OutNumberOfSectionsY * OutNumberOfSectionsZ;
-	}
-
-	TArray<MeshPartition::FMeshData> BuildSections(const MeshPartition::FMeshData& InMesh, const double InSectionSize)
+	TArray<MeshPartition::FMeshData> BuildSections(const MeshPartition::FMeshData& InMesh, const double InSectionSize, const bool bInFilterEmptyMeshes)
 	{
 		//todo(luc.eygasier): think about reusing the spatial structure in the builder result.
 		MeshPartition::FMeshABBTree3 Spatial;
@@ -483,14 +520,17 @@ namespace UE::MeshPartition::BuildHelpers
 
 		Spatial.SetMesh(&InMesh, bAutoBuild);
 
-		int32 NumberOfMeshSectionsX;
-		int32 NumberOfMeshSectionsY;
-		int32 NumberOfMeshSectionsZ;
-
 		const FBox MeshBounds(InMesh.GetBounds());
-		const UE::Geometry::FAxisAlignedBox3d AlignedMeshBounds(MeshBounds);
-		const int32 TotalNumberOfSections = GetTotalNumberOfSections(MeshBounds, InSectionSize, NumberOfMeshSectionsX, NumberOfMeshSectionsY, NumberOfMeshSectionsZ);
-			
+		const GridHelpers::FGridDimensions Grid = GridHelpers::ComputeGridDimensions(MeshBounds, InSectionSize);
+
+		const int32 NumberOfMeshSectionsX = Grid.CellNumber.X;
+		const int32 NumberOfMeshSectionsY = Grid.CellNumber.Y;
+		const int32 NumberOfMeshSectionsZ = Grid.CellNumber.Z;
+		const int32 TotalNumberOfSections = Grid.TotalCells;
+
+		const FBox SnappedMeshBounds(Grid.SnappedMin, MeshBounds.Max);
+		const UE::Geometry::FAxisAlignedBox3d AlignedMeshBounds(SnappedMeshBounds);
+
 		TArray<MeshPartition::FMeshData> Result;
 
 		Result.SetNum(TotalNumberOfSections);
@@ -506,18 +546,17 @@ namespace UE::MeshPartition::BuildHelpers
 				const int32 SectionIndexX = SectionIndex % NumberOfMeshSectionsX;
 				const int32 SectionIndexY = SectionIndex / NumberOfMeshSectionsX % NumberOfMeshSectionsY;
 				const int32 SectionIndexZ = SectionIndex / (NumberOfMeshSectionsX * NumberOfMeshSectionsY);
-		
+
 				const FVector SectionPos = AlignedMeshBounds.Min + FVector(SectionIndexX, SectionIndexY, SectionIndexZ) * InSectionSize;
-		
+
 				const UE::Geometry::FAxisAlignedBox3d SectionBox(SectionPos, SectionPos + InSectionSize);
 				const UE::Geometry::FAxisAlignedBox3d ExpandedSectionBox(SectionPos - KINDA_SMALL_NUMBER, SectionPos + InSectionSize + KINDA_SMALL_NUMBER);
-		
+
 				MeshPartition::FMeshABBTree3::FTreeTraversal Traversal;
 
 				int SelectAllDepth = TNumericLimits<int>::Max();
 				int CurrentDepth = -1;
-		
-		
+
 				{
 					Traversal.NextBoxF = [&ExpandedSectionBox, &CurrentDepth, &SelectAllDepth](const UE::Geometry::FAxisAlignedBox3d& Box, int Depth)
 					{
@@ -538,7 +577,7 @@ namespace UE::MeshPartition::BuildHelpers
 
 						return ExpandedSectionBox.Intersects(Box);
 					};
-		
+
 					Traversal.NextTriangleF = [&TriangleOwners, &SectionBox, &InMesh, &CurrentDepth, &SelectAllDepth, SectionIndex = SectionIndex] (int TriangleID)
 					{
 						auto MarkTriangle = [&](int TriangleID)
@@ -565,7 +604,7 @@ namespace UE::MeshPartition::BuildHelpers
 						FVector3d VertexA;
 						FVector3d VertexB;
 						FVector3d VertexC;
-			
+
 						InMesh.GetTriVertices(TriangleID, VertexA, VertexB, VertexC);
 
 						FVector3d Centroid = (VertexA + VertexB + VertexC) / 3.0;
@@ -575,7 +614,7 @@ namespace UE::MeshPartition::BuildHelpers
 							MarkTriangle(TriangleID);
 						}
 					};
-		
+
 					Spatial.DoTraversal(Traversal);
 				}
 			});
@@ -598,7 +637,7 @@ namespace UE::MeshPartition::BuildHelpers
 				TMap<int, int> SourceToNewVertexIDs;
 				TArray<int> NewToSourceVertexIDs;
 				TArray<int> NewToSourceTriangleIDs;
-		
+
 				MeshPartition::FMeshData& ResultMesh = Result[SectionIndex];
 
 				for (int SourceTriangleID : Triangles)
@@ -611,12 +650,12 @@ namespace UE::MeshPartition::BuildHelpers
 						const int VertexID = Triangle[Index];
 						int NewVertexID = 0;
 						int* NewVertexIDPtr = SourceToNewVertexIDs.Find(VertexID);
-				
+
 						if (NewVertexIDPtr == nullptr)
 						{
 							NewVertexID = ResultMesh.AppendVertex(InMesh.GetVertex(VertexID));
 							SourceToNewVertexIDs.Emplace(VertexID, NewVertexID);
-					
+
 							int NewVertexIDIndex = NewToSourceVertexIDs.Emplace(VertexID);
 							ensure(NewVertexID == NewVertexIDIndex);
 						}
@@ -624,22 +663,25 @@ namespace UE::MeshPartition::BuildHelpers
 						{
 							NewVertexID = *NewVertexIDPtr;
 						}
-		
+
 						NewTriangle[Index] = NewVertexID;
 					}
-		
+
 					int NewTriangleID = ResultMesh.AppendTriangle(NewTriangle);
 					int NewTriangleIDIndex = NewToSourceTriangleIDs.Emplace(SourceTriangleID);
 					ensure(NewTriangleID == NewTriangleIDIndex);
 				}
-		
+
 				constexpr bool bTransferNormals = true;
 				BuildHelpers::Private::TransferVertexAttributes(NewToSourceVertexIDs, InMesh, ResultMesh, bTransferNormals);
 				BuildHelpers::Private::TransferTriangleAttributes(NewToSourceTriangleIDs, InMesh, ResultMesh);
 			});
 		}
 
-		Result.RemoveAll([](const MeshPartition::FMeshData& Mesh) { return Mesh.VertexCount() == 0; });
+		if (bInFilterEmptyMeshes)
+		{
+			Result.RemoveAll([](const MeshPartition::FMeshData& Mesh) { return Mesh.VertexCount() == 0; });
+		}
 		return MoveTemp(Result);
 	}
 
diff --git a/Engine/Plugins/Experimental/MeshPartition/Source/MeshPartitionEditor/Private/MeshPartitionRuntimeCellTransformer.cpp b/Engine/Plugins/Experimental/MeshPartition/Source/MeshPartitionEditor/Private/MeshPartitionRuntimeCellTransformer.cpp
--- a/Engine/Plugins/Experimental/MeshPartition/Source/MeshPartitionEditor/Private/MeshPartitionRuntimeCellTransformer.cpp
+++ b/Engine/Plugins/Experimental/MeshPartition/Source/MeshPartitionEditor/Private/MeshPartitionRuntimeCellTransformer.cpp
@@ -180,6 +180,13 @@ void URuntimeCellTransformer::Transform(ULevel* InLevel)
 						return true;
 					}
 
+					if (!Status->bBuildVariantHashMatches)
+					{
+						MEGAMESH_DEBUG_LOG(TEXT("    REMOVING BuildVariantHashFail %.6s (%p %s)"), *CompiledSection->GetActorGuid().ToString(), Actor.Get(), *Actor->GetActorNameOrLabel());
+						MeshPartitionUpdater->RemovedOutOfDate++;
+						return true;
+					}
+
 					if (!Status->bPackageHashMatches)
 					{
 						MEGAMESH_DEBUG_LOG(TEXT("    REMOVING PackageHashFail %.6s (%p %s)"), *CompiledSection->GetActorGuid().ToString(), Actor.Get(), *Actor->GetActorNameOrLabel());
diff --git a/Engine/Plugins/Experimental/MeshPartition/Source/MeshPartitionEditor/Private/MeshPartitionSeparateWorldBuilder.cpp b/Engine/Plugins/Experimental/MeshPartition/Source/MeshPartitionEditor/Private/MeshPartitionSeparateWorldBuilder.cpp
--- a/Engine/Plugins/Experimental/MeshPartition/Source/MeshPartitionEditor/Private/MeshPartitionSeparateWorldBuilder.cpp
+++ b/Engine/Plugins/Experimental/MeshPartition/Source/MeshPartitionEditor/Private/MeshPartitionSeparateWorldBuilder.cpp
@@ -1,31 +1,186 @@
 // Copyright Epic Games, Inc. All Rights Reserved.
 
 #include "MeshPartitionSeparateWorldBuilder.h"
+
+#include "Algo/IsSorted.h"
 #include "MeshPartition.h"
-#include "MeshPartitionEditorModule.h"
-#include "MeshPartitionEditorSubsystem.h"
+#include "MeshPartitionChannelCollection.h"
 #include "MeshPartitionCompiledSection.h"
+#include "MeshPartitionDescriptorCache.h"
 #include "MeshPartitionEditorComponent.h"
+#include "MeshPartitionEditorModule.h"
+#include "MeshPartitionEditorSubsystem.h"
+#include "MeshPartitionMeshBuilder.h"
 #include "MeshPartitionModifierComponent.h"
-#include "MeshPartitionDescriptorCache.h"
-#include "WorldPartition/WorldPartitionSubsystem.h"
-#include "WorldPartition/WorldPartitionHandle.h"
-#include "WorldPartition/WorldPartition.h"
-#include "WorldPartition/WorldPartitionHelpers.h"
 #include "MeshPartitionWorldPartitionHelpers.h"
 #include "EditorWorldUtils.h"
 #include "EngineUtils.h"
 #include "MaterialCache/MaterialCache.h"
 #include "UObject/Package.h"
+#include "WorldPartition/WorldPartition.h"
+#include "WorldPartition/WorldPartitionHandle.h"
+#include "WorldPartition/WorldPartitionHelpers.h"
+#include "WorldPartition/WorldPartitionSubsystem.h"
+
+static TAutoConsoleVariable<bool> CVarFillModifiersToProcess(
+	TEXT("MeshPartition.SeparateWorldBuilder.FillModifiersToProcess"),
+	true,
+	TEXT("When false, skips filling ModifiersToProcess during PIE build (debug only)."));
 
 namespace UE::MeshPartition
 {
+
+struct FModifierGroupBuildCacheKey
+{
+	FGuid MeshPartitionGUID;
+	FName BuildVariantName;
+	TArray<FSoftObjectPath> BaseModifierPaths;
+
+	bool operator==(const FModifierGroupBuildCacheKey& InOther) const
+	{
+		return MeshPartitionGUID == InOther.MeshPartitionGUID
+			&& BuildVariantName == InOther.BuildVariantName
+			&& BaseModifierPaths == InOther.BaseModifierPaths;
+	}
+
+	friend uint32 GetTypeHash(const FModifierGroupBuildCacheKey& InKey)
+	{
+		return HashCombine(HashCombine(GetTypeHash(InKey.MeshPartitionGUID), GetTypeHash(InKey.BuildVariantName)), HashBaseModifierPaths(InKey.BaseModifierPaths));
+	}
+
+	// Compute a stable hash from the sorted base modifier paths on the BuildInfo.
+	static uint32 HashBaseModifierPaths(const TArray<FSoftObjectPath>& InBaseModifierPaths)
+	{
+		// HashCombine is order-dependent; paths must be sorted for deterministic results.
+		ensure(Algo::IsSorted(InBaseModifierPaths, [](const FSoftObjectPath& A, const FSoftObjectPath& B) { return A.GetAssetPath().Compare(B.GetAssetPath()) < 0; }));
+
+		uint32 Hash = 0;
+		for (const FSoftObjectPath& Path : InBaseModifierPaths)
+		{
+			Hash = HashCombine(Hash, GetTypeHash(Path));
+		}
+		return Hash;
+	}
+};
+
+struct FCachedModifierGroupBuild
+{
+	TSharedPtr<const FSectionChannels> Channels;
+	TMap<FIntVector, TSharedPtr<const FMeshData>> CellMeshes;		// Keyed by absolute world-space grid coordinate
+	FModifierGroup ModifierGroup;
+	FGuid ModifiersHash;										// Computed once on cache miss, reused for subsequent siblings
+	TSet<FWorldPartitionReference> ActorRefs;			// Keeps modifier actors loaded while cache is alive
+};
+
+struct FModifierGroupBuildCache
+{
+	TMap<FModifierGroupBuildCacheKey, FCachedModifierGroupBuild> Entries;
+};
+
+namespace MeshPartitionSeparateWorldBuilderLocals
+{
+
+void ApplyGridCellToSection(ACompiledSection*					InCompiledSection,
+							const FCachedModifierGroupBuild&	InCachedBuild,
+							UMeshPartitionEditorComponent*		InEditorComponent)
+{
+	TRACE_CPUPROFILER_EVENT_SCOPE(MeshPartition::SeparateWorldBuilder::ApplyGridCellToSection);
+
+	const FCompiledSectionBuildInfo& BuildInfo = InCompiledSection->GetBuildInfo();
+	const FIntVector& GridCellCoord = BuildInfo.GridCellCoord;
+
+	const TSharedPtr<const FMeshData>* CellMeshPtr = InCachedBuild.CellMeshes.Find(GridCellCoord);
+
+	if (!CellMeshPtr || !*CellMeshPtr)
+	{
+		UE_LOG(LogMegaMeshEditor, Verbose, TEXT("No geometry for grid cell (%d,%d,%d) — leaving as placeholder."), GridCellCoord.X, GridCellCoord.Y, GridCellCoord.Z);
+		return;
+	}
+
+	if (!ensure(InCachedBuild.Channels.IsValid()))
+	{
+		UE_LOG(LogMegaMeshEditor, Warning, TEXT("No section channels for grid cell (%d,%d,%d) — leaving as placeholder."), GridCellCoord.X, GridCellCoord.Y, GridCellCoord.Z);
+		return;
+	}
+
+	InEditorComponent->ApplyBuiltMeshToSection(InCompiledSection, *CellMeshPtr, *InCachedBuild.Channels, InCachedBuild.ModifierGroup);
+	InCompiledSection->SetIsPlaceholder(false);
+
+	UE_LOG(LogMegaMeshEditor, Verbose, TEXT("Built grid cell (%d,%d,%d): %p %s (%d vertices, %d tris)"),
+		   GridCellCoord.X, GridCellCoord.Y, GridCellCoord.Z, InCompiledSection, *InCompiledSection->GetActorNameOrLabel(),
+		   (*CellMeshPtr)->VertexCount(), (*CellMeshPtr)->TriangleCount());
+}
+
+// Splits the full mesh into grid cells, stores the result in the build cache, and
+// returns a pointer to the cached entry. Returns nullptr if the split could not be
+// performed (caller falls through to single-section path).
+FCachedModifierGroupBuild* SplitAndCacheGridBuild(const FCompiledSectionBuildInfo&				InBuildInfo,
+												  const FBuildTaskHandle&						InTaskHandle,
+												  FModifierGroup&&								InModifierGroup,
+												  TSet<FWorldPartitionReference>&&				InActorRefs,
+												  FModifierGroupBuildCache&						InCache)
+{
+	TRACE_CPUPROFILER_EVENT_SCOPE(MeshPartition::SeparateWorldBuilder::SplitAndCacheGridBuild);
+
+	// Read from BuildInfo — set by the WorldUpdater when creating placeholders.
+	// ResolveGridCellSizeFromPipeline requires access to WP RuntimeHash, which is not
+	// available in the PIE world; the editor world resolved and stamped it at placeholder creation.
+	const uint32 ResolvedGridCellSize = InBuildInfo.ResolvedGridCellSize;
+	TSharedPtr<const FMeshData> ProcessedMesh = InTaskHandle.GetTask()->GetMesh();
+
+	UE_LOG(LogMegaMeshEditor, Verbose, TEXT("Grid-split build: GridCellCoord=(%d,%d,%d) ResolvedGridCellSize=%u MeshBounds=[%.1f,%.1f,%.1f]-[%.1f,%.1f,%.1f]"),
+		   InBuildInfo.GridCellCoord.X, InBuildInfo.GridCellCoord.Y, InBuildInfo.GridCellCoord.Z, ResolvedGridCellSize,
+		   ProcessedMesh ? ProcessedMesh->GetBounds().Min.X : 0.0, ProcessedMesh ? ProcessedMesh->GetBounds().Min.Y : 0.0, ProcessedMesh ? ProcessedMesh->GetBounds().Min.Z : 0.0,
+		   ProcessedMesh ? ProcessedMesh->GetBounds().Max.X : 0.0, ProcessedMesh ? ProcessedMesh->GetBounds().Max.Y : 0.0, ProcessedMesh ? ProcessedMesh->GetBounds().Max.Z : 0.0);
+
+	if ((ResolvedGridCellSize == 0) || !ProcessedMesh)
+	{
+		UE_LOG(LogMegaMeshEditor, Warning, TEXT("Grid cell size could not be resolved for grid-split build — falling back to single-section."));
+		return nullptr;
+	}
+
+	// Split the full mesh into grid cells keyed by absolute world-space coordinates.
+	// Absolute coordinates decouple this path from the WorldUpdater's placeholder creation:
+	// both derive cell coordinates from world position and cell size alone, independent of
+	// the bounding box used to enumerate cells. Extra cells (from modifiers growing the mesh
+	// beyond ComputeBaseBounds) simply won't match any placeholder and are harmlessly skipped.
+	TMap<FIntVector, FMeshData> RawCellMeshes = GridHelpers::BuildGridCellMeshes(*ProcessedMesh, static_cast<double>(ResolvedGridCellSize));
+
+	UE_LOG(LogMegaMeshEditor, Verbose, TEXT("Grid-split produced %d non-empty cells (GridCellCoord=(%d,%d,%d))."),
+		   RawCellMeshes.Num(), InBuildInfo.GridCellCoord.X, InBuildInfo.GridCellCoord.Y, InBuildInfo.GridCellCoord.Z);
+
+	FCachedModifierGroupBuild NewCache;
+	NewCache.Channels = InTaskHandle.GetTask()->GetSectionChannels();
+	NewCache.CellMeshes.Reserve(RawCellMeshes.Num());
+
+	for (TPair<FIntVector, FMeshData>& CellEntry : RawCellMeshes)
+	{
+		NewCache.CellMeshes.Add(CellEntry.Key, MakeShared<const FMeshData>(MoveTemp(CellEntry.Value)));
+	}
+
+	NewCache.ModifierGroup = MoveTemp(InModifierGroup);
+	NewCache.ModifiersHash = InBuildInfo.ModifiersHash;
+	NewCache.ActorRefs = MoveTemp(InActorRefs);
+
+	FModifierGroupBuildCacheKey CacheKey{InBuildInfo.MegaMeshGUID, InBuildInfo.BuildVariantName, InBuildInfo.BaseModifierPaths};
+	return &InCache.Entries.Add(CacheKey, MoveTemp(NewCache));
+}
+
+} // namespace MeshPartitionSeparateWorldBuilderLocals
+
+// ---------------------------------------------------------------------------
+// FSeparateWorldBuilder
+// ---------------------------------------------------------------------------
+
 void FSeparateWorldBuilder::InitializeToBuildOutOfWorldPackage(const FString& InWorldPackageName)
 {
+	TRACE_CPUPROFILER_EVENT_SCOPE(MeshPartition::SeparateWorldBuilder::InitializeToBuildOutOfWorldPackage);
+
 	// unload existing world if necessary
 	if (World != nullptr)
 	{
 		FString ExistingWorldPackageName = World->GetPackage()->GetPathName();
+
 		if (ExistingWorldPackageName != InWorldPackageName)
 		{
 			UE_LOG(LogMegaMeshEditor, Verbose, TEXT("InitializeToBuildOutOfWorldPackage: Changing Worlds '%s' ==> '%s'"), *ExistingWorldPackageName, *InWorldPackageName);
@@ -54,6 +209,7 @@ void FSeparateWorldBuilder::InitializeToBuildOutOfWorldPackage(const FString& In
 	{
 		BuilderPackage = CreatePackage(TEXT("/Temp/MeshPartitionBuilderPackage"));
 	}
+
 	ensure(BuilderPackage);
 
 	if (World == nullptr)
@@ -67,225 +223,282 @@ void FSeparateWorldBuilder::InitializeToBuildOutOfWorldPackage(const FString& In
 		FArchive* ReaderArchive = nullptr;
 		FLinkerInstancingContext* LinkerInstancingContext = nullptr;
 		UPackage* MapPackage = LoadPackage(BuilderPackage, LongPackageName, LoadFlags, ReaderArchive, LinkerInstancingContext);
+
 		check(MapPackage);
 
 		World = UWorld::FindWorldInPackage(MapPackage);
+
 		check(!World->IsPlayInEditor());
 	}
+
 	ensure(World);
 }
 
-void FSeparateWorldBuilder::BuildCompiledSectionMesh(MeshPartition::ACompiledSection* CompiledSection, FName BuildVariantName)
+void FSeparateWorldBuilder::BuildCompiledSectionMesh(MeshPartition::ACompiledSection* InCompiledSection, FName InBuildVariantName)
 {
+	TRACE_CPUPROFILER_EVENT_SCOPE(MeshPartition::SeparateWorldBuilder::BuildCompiledSectionMesh);
+
 	// NOTE: the compiled section is the actor we need to build the mesh for
 	// however, the compiled section does NOT belong to our World (our World is a temp world loaded just for building)
-	check(CompiledSection);
-	UWorld* CompiledSectionWorld = CompiledSection->GetWorld();
+	check(InCompiledSection);
+	UWorld* CompiledSectionWorld = InCompiledSection->GetWorld();
+
+	if (CompiledSectionWorld == nullptr)
+	{
+		// The compiled section's level may have been streamed out. The actor is still
+		// reachable via weak pointer but its world is gone — skip it.
+		UE_LOG(LogMegaMeshEditor, Warning, TEXT("BuildCompiledSectionMesh: compiled section %s has no world — skipping (likely streamed out)."), *InCompiledSection->GetActorNameOrLabel());
+		return;
+	}
 
-	MeshPartition::FCompiledSectionBuildInfo BuildInfo = CompiledSection->GetBuildInfo();
+	MeshPartition::FCompiledSectionBuildInfo BuildInfo = InCompiledSection->GetBuildInfo();
+
+	// Fix PIE paths early so that all downstream uses (cache keys, modifier group building)
+	// operate on consistent editor-world paths.
+	{
+		MeshPartition::EditorUtils::FPIEPathFixer& PIEPathFixer = UMeshPartitionEditorSubsystem::GetPIEPathFixer();
+		for (FSoftObjectPath& ModifierPath : BuildInfo.BaseModifierPaths)
+		{
+			PIEPathFixer.FixInPlace(ModifierPath);
+		}
+
+		// Re-sort after PIE-fixing: FixInPlace changes path strings, which can break the
+		// lexicographic order that HashBaseModifierPaths requires for deterministic hashing.
+		BuildInfo.BaseModifierPaths.Sort([](const FSoftObjectPath& A, const FSoftObjectPath& B) { return A.GetAssetPath().Compare(B.GetAssetPath()) < 0; });
+	}
 
 	// Get the mesh partition definition.  Note that the compiled section is in a streaming cell, and it's parent mesh partition may not be loaded
 	// So instead we make sure to grab the definition from the path stored on the buildinfo
 	FSoftObjectPath MeshPartitionDefinitionPath(BuildInfo.MegaMeshDefinitionPath);
 	TSoftObjectPtr<UMeshPartitionDefinition> MeshPartitionDefinition(MeshPartitionDefinitionPath);
 	UMeshPartitionDefinition* Definition = MeshPartitionDefinition.Get();
+
 	if (Definition == nullptr)
 	{
 		Definition = UMeshPartitionDefinition::StaticClass()->GetDefaultObject<UMeshPartitionDefinition>();
 	}
 
-	const MeshPartition::FCompiledSectionBuildVariant& BuildVariant = Definition->GetCompiledSectionBuildVariantByName(BuildVariantName);
+	const MeshPartition::FCompiledSectionBuildVariant& BuildVariant = Definition->GetCompiledSectionBuildVariantByName(InBuildVariantName);
 
-	if (MeshPartition::FMeshBuilder* Builder = UMeshPartitionEditorSubsystem::GetMeshBuilder())
+	MeshPartition::FMeshBuilder* Builder = UMeshPartitionEditorSubsystem::GetMeshBuilder();
+
+	if (Builder == nullptr)
 	{
-		MeshPartition::FModifierDescriptorCache* DescriptorCache = UMeshPartitionEditorSubsystem::GetDescriptorCache();
-		check(DescriptorCache);
+		return;
+	}
 
-		// We search by GUID because the path can be modified by PIE
-		TSharedPtr<const MeshPartition::FCachedDescriptors> Descriptors = DescriptorCache->GetCachedDescriptorsByGuid(BuildInfo.MegaMeshGUID);
+	MeshPartition::FModifierDescriptorCache* DescriptorCache = UMeshPartitionEditorSubsystem::GetDescriptorCache();
 
-		FSoftObjectPath MeshPartitionActorPath = Descriptors->GetMeshPartitionPath();
+	check(DescriptorCache);
 
-		// build the section group that corresponds to these base modifiers (after first fixing the PIE path names..)
-		MeshPartition::EditorUtils::FPIEPathFixer& PIEPathFixer = UMeshPartitionEditorSubsystem::GetPIEPathFixer();
-		for (FSoftObjectPath& ModifierPath : BuildInfo.BaseModifierPaths)
+	// We search by GUID because the path can be modified by PIE
+	TSharedPtr<const MeshPartition::FCachedDescriptors> Descriptors = DescriptorCache->GetCachedDescriptorsByGuid(BuildInfo.MegaMeshGUID);
+
+	if (!Descriptors.IsValid())
+	{
+		UE_LOG(LogMegaMeshEditor, Warning, TEXT("Failed to build compiled section: no cached descriptors for MeshPartition GUID %s"), *BuildInfo.MegaMeshGUID.ToString());
+		return;
+	}
+
+	// find the mesh partition actor by GUID (NOTE: we can't find it by path, unless we use the PIE world path, otherwise it will get the mesh partition actor in the main editor world)
+	FGuid MeshPartitionActorGUID = Descriptors->GetMeshPartitionGUID();
+	const AMeshPartition* MeshPartitionActor = nullptr;
+
+	for (AMeshPartition* MeshPartitionActorIt : TActorRange<AMeshPartition>(CompiledSectionWorld))
+	{
+		if (MeshPartitionActorIt && (MeshPartitionActorIt->GetActorGuid() == MeshPartitionActorGUID))
 		{
-			PIEPathFixer.FixInPlace(ModifierPath);
+			MeshPartitionActor = MeshPartitionActorIt;
+			break;
 		}
-		MeshPartition::FModifierGroup ModifierGroup = Descriptors->BuildModifierGroupForBaseModifierPaths(BuildInfo.BaseModifierPaths);
+	}
+
+	if (MeshPartitionActor == nullptr)
+	{
+		UE_LOG(LogMegaMeshEditor, Warning, TEXT("Failed to build compiled section: could not find parent MeshPartition actor %s in World %s"), *MeshPartitionActorGUID.ToString(), *CompiledSectionWorld->GetName());
+		return;
+	}
+
+	UMeshPartitionEditorComponent* MeshPartitionEditorComponent = Cast<UMeshPartitionEditorComponent>(MeshPartitionActor->GetMeshPartitionComponent());
+
+	check(MeshPartitionEditorComponent);
+	check(MeshPartitionEditorComponent->GetWorld() == CompiledSectionWorld);
 
-		// find the mesh partition actor by GUID (NOTE: we can't find it by path, unless we use the PIE world path, otherwise it will get the mesh partition actor in the main editor world)
-		FGuid MeshPartitionActorGUID = Descriptors->GetMeshPartitionGUID();
-		const AMeshPartition* MeshPartitionActor = nullptr;
-		for (AMeshPartition* MeshPartition : TActorRange<AMeshPartition>(CompiledSectionWorld))
+	// --- Grid-split cache: check for an existing build before doing expensive work ---
+	if (BuildVariant.bSplitSectionsToMatchWorldPartitionRuntimeGrid)
+	{
+		if (!GroupBuildCache)
 		{
-			if (MeshPartition && (MeshPartition->GetActorGuid() == MeshPartitionActorGUID))
-			{
-				MeshPartitionActor = MeshPartition;
-				break;
-			}
+			GroupBuildCache = MakeShared<FModifierGroupBuildCache>();
 		}
 
-		if (MeshPartitionActor == nullptr)
+		FModifierGroupBuildCacheKey CacheKey{BuildInfo.MegaMeshGUID, InBuildVariantName, BuildInfo.BaseModifierPaths};
+
+		if (FCachedModifierGroupBuild* CachedBuild = GroupBuildCache->Entries.Find(CacheKey))
 		{
-			UE_LOG(LogMegaMeshEditor, Warning, TEXT("Failed to build compiled section: could not find parent MeshPartition actor %s in World %s"), *MeshPartitionActorGUID.ToString(), *CompiledSectionWorld->GetName());
+			UE_LOG(LogMegaMeshEditor, Verbose, TEXT("Grid build cache hit for group %s — skipping mesh rebuild."), *BuildInfo.ModifierSetHash.ToString());
+
+			// Stamp the PIE-fixed BuildInfo with the cached ModifiersHash so the section
+			// carries correct metadata (matching what the first sibling received on cache miss).
+			BuildInfo.ModifiersHash = CachedBuild->ModifiersHash;
+			InCompiledSection->SetBuildInfo(BuildInfo);
+
+			MeshPartitionSeparateWorldBuilderLocals::ApplyGridCellToSection(InCompiledSection, *CachedBuild, MeshPartitionEditorComponent);
 			return;
 		}
+	}
 
-		UMeshPartitionEditorComponent* MeshPartitionEditorComponent = Cast<UMeshPartitionEditorComponent>(MeshPartitionActor->GetMeshPartitionComponent());
-		check(MeshPartitionEditorComponent);
-		check(MeshPartitionEditorComponent->GetWorld() == CompiledSectionWorld);
+	// --- Full mesh build (shared between grid miss and non-grid paths) ---
 
-		auto OnActorLoaded = [](AActor* Actor, bool bIsInLevelInstance)
-			{
-				// In the commandlet build path we just record the actor refs to ensure they are not GCed while we are working.
-				// As this path is synchronous currently, it's not necessary,
-			};
+	MeshPartition::FModifierGroup ModifierGroup = Descriptors->BuildModifierGroupForBaseModifierPaths(BuildInfo.BaseModifierPaths);
 
-		auto OnModifierLoaded = [](MeshPartition::UModifierComponent* Modifier, bool bIsInLevelInstance)
-			{
-				Modifier->InitializeModifier();
-			};
+	auto OnActorLoaded = [](AActor* InActor, bool bInIsInLevelInstance)
+		{
+			// In the commandlet build path we just record the actor refs to ensure they are not GCed while we are working.
+			// As this path is synchronous currently, it's not necessary,
+		};
+
+	auto OnModifierLoaded = [](MeshPartition::UModifierComponent* InModifier, bool bInIsInLevelInstance)
+		{
+			InModifier->InitializeModifier();
+		};
+
+	TSet<FWorldPartitionReference> ActorRefs;
 
-		TSet<FWorldPartitionReference> ActorRefs;
+	{
+		TRACE_CPUPROFILER_EVENT_SCOPE(MeshPartition::SeparateWorldBuilder::LoadModifiers);
 		WorldPartitionHelpers::LoadAllRelevantMegaMeshActors(ModifierGroup, MeshPartitionActor, ActorRefs, OnActorLoaded, OnModifierLoaded);
+	}
+
+	FGuid GroupHash = ModifierGroup.UpdateAndComputeModifierGroupHash();
+	BuildInfo.ModifiersHash = GroupHash;
 
-		FGuid GroupHash = ModifierGroup.UpdateAndComputeModifierGroupHash();
-		BuildInfo.ModifiersHash = GroupHash;
+	// Stamp the PIE-fixed, hash-updated BuildInfo on the section so that downstream code
+	// (ApplyBuiltMeshToSection) can read it directly via GetBuildInfo().
+	InCompiledSection->SetBuildInfo(BuildInfo);
 
-		MeshPartition::FBuilderSettings BuilderSettings;
+	MeshPartition::FBuilderSettings BuilderSettings;
+	{
+		BuilderSettings.BuildType = MeshPartition::EBuildType::CompiledSection;
+		BuilderSettings.Transform = FTransform::Identity;
+
+		// Fill out modifiers to process
+		BuilderSettings.ModifiersToProcess.Empty();
+
+		if (CVarFillModifiersToProcess.GetValueOnGameThread())
 		{
-			BuilderSettings.BuildType = MeshPartition::EBuildType::CompiledSection;
-			BuilderSettings.Transform = FTransform::Identity;
-			
-			// Fill out modifiers to process
-			static bool bFillModifiersToProcess = true;
-			BuilderSettings.ModifiersToProcess.Empty();
-			
-			if (bFillModifiersToProcess)
+			ModifierGroup.ForAllModifiers([&BuilderSettings](MeshPartition::UModifierComponent* InModifierComponent)
 			{
-				ModifierGroup.ForAllModifiers([&BuilderSettings](MeshPartition::UModifierComponent* ModifierComponent)
+				if (ensure(InModifierComponent != nullptr))
 				{
-					if (ensure(ModifierComponent != nullptr))
-					{
-						BuilderSettings.ModifiersToProcess.Add(ModifierComponent);
-					}
-
-					return true;
-				});
-			}
+					BuilderSettings.ModifiersToProcess.Add(InModifierComponent);
+				}
 
-			BuilderSettings.TypePriorities = Definition->GetModifierTypePriorities();
-			// BuilderSettings.FilterBounds = ;				// optional bounds
-			// BuilderSettings.ModifierFilter = ...;		// Optional filter function for the list of modifiers to process.
-			BuilderSettings.MaxSectionComplexity = BuildVariant.MaxSectionComplexity;
-			BuilderSettings.bRecomputeNormals = false;
-			BuilderSettings.bRecomputeTangents = false;		// not yet supported
-			BuilderSettings.bCacheResult = true;
-			BuilderSettings.bBuildSpatial = false;
-			BuilderSettings.bAllowDDCRead = true;			// Should builds with these settings be allowed to use the DDC cache to avoid processing.
-			BuilderSettings.bAllowDDCWrite = true;			// Should builds produced with these settings be allowed to store their results in DDC.
-			BuilderSettings.TexcoordGenerationOptions = FChannelCollectionUVLayoutOptions::GetFromDefinition(Definition);
-			// Choose the higher quality layout method for compiled sections
-			BuilderSettings.TexcoordGenerationOptions->UVLayoutMethod = EChannelCollectionUVLayoutMethod::ReferenceBoxProject;
-
-			MeshPartition::FBuilderSettings::FChannelRenderSettings ChannelRenderSettings;
-			ChannelRenderSettings.ChannelMap = Definition->GetChannelMap();
-			ChannelRenderSettings.TexelSize = Definition->GetChannelTexelSize();
-			BuilderSettings.ChannelRenderSettings.Emplace(MoveTemp(ChannelRenderSettings));
+				return true;
+			});
 		}
 
-		const bool bAllowDDC = false;
-		MeshPartition::FBuildTaskHandle TaskHandle = Builder->Build(BuilderSettings, ModifierGroup, bAllowDDC);
+		BuilderSettings.TypePriorities = Definition->GetModifierTypePriorities();
+		// BuilderSettings.FilterBounds = ;				// optional bounds
+		// BuilderSettings.ModifierFilter = ...;		// Optional filter function for the list of modifiers to process.
+		BuilderSettings.MaxSectionComplexity = BuildVariant.MaxSectionComplexity;
+		BuilderSettings.bRecomputeNormals = false;
+		BuilderSettings.bRecomputeTangents = false;		// not yet supported
+		BuilderSettings.bCacheResult = true;
+		BuilderSettings.bBuildSpatial = false;
+		BuilderSettings.bAllowDDCRead = true;			// Should builds with these settings be allowed to use the DDC cache to avoid processing.
+		BuilderSettings.bAllowDDCWrite = true;			// Should builds produced with these settings be allowed to store their results in DDC.
+		BuilderSettings.TexcoordGenerationOptions = FChannelCollectionUVLayoutOptions::GetFromDefinition(Definition);
+		// Choose the higher quality layout method for compiled sections
+		BuilderSettings.TexcoordGenerationOptions->UVLayoutMethod = EChannelCollectionUVLayoutMethod::ReferenceBoxProject;
+
+		MeshPartition::FBuilderSettings::FChannelRenderSettings ChannelRenderSettings;
+		ChannelRenderSettings.ChannelMap = Definition->GetChannelMap();
+		ChannelRenderSettings.TexelSize = Definition->GetChannelTexelSize();
+		BuilderSettings.ChannelRenderSettings.Emplace(MoveTemp(ChannelRenderSettings));
+	}
+
+	const bool bAllowDDC = false;
+	MeshPartition::FBuildTaskHandle TaskHandle;
 
-		//TaskHandles.Add(TaskHandle);
+	{
+		TRACE_CPUPROFILER_EVENT_SCOPE(MeshPartition::SeparateWorldBuilder::BuildMesh);
+		TaskHandle = Builder->Build(BuilderSettings, ModifierGroup, bAllowDDC);
 		TaskHandle.Wait();
-		check(TaskHandle.IsCompleted());
+	}
 
-		TSharedPtr<const MeshPartition::FSectionChannels> SectionChannel = TaskHandle.GetTask()->GetSectionChannels();
+	check(TaskHandle.IsCompleted());
 
-		const bool bUseStaticMobility = true;
-		CompiledSection = MeshPartitionEditorComponent->PrepareMegaMeshCompiledSection(BuildInfo, BuildVariant, CompiledSection, bUseStaticMobility);
+	// --- Grid-split: split mesh, cache, and apply the requested cell ---
+	if (BuildVariant.bSplitSectionsToMatchWorldPartitionRuntimeGrid)
+	{
+		FCachedModifierGroupBuild* CachedBuild = MeshPartitionSeparateWorldBuilderLocals::SplitAndCacheGridBuild(BuildInfo, TaskHandle, MoveTemp(ModifierGroup), MoveTemp(ActorRefs), *GroupBuildCache);
 
-		CompiledSection->SetChannelTexture(SectionChannel->Texture.Get());
-		CompiledSection->SetChannelTexcoordDesc(SectionChannel->TexcoordMetrics);
-		CompiledSection->SetChannelTable(SectionChannel->Table);
-		
-		if (SectionChannel->Domain.AreaUV > 0)
+		if (CachedBuild != nullptr)
 		{
-			CompiledSection->SetMaterialCacheTileCount(GetMaterialCacheTileCount(
-				SectionChannel->Domain.Area3D / SectionChannel->Domain.AreaUV,
-				Definition->GetMaterialCacheTexelSize()
-			));
-		
-			CompiledSection->RecreateMaterialCacheTextures();
+			MeshPartitionSeparateWorldBuilderLocals::ApplyGridCellToSection(InCompiledSection, *CachedBuild, MeshPartitionEditorComponent);
 		}
-
-		TArray<TWeakObjectPtr<MeshPartition::UModifierComponent>> ModifierPtrs = ModifierGroup.AllResolvedModifierPtrs();
-
-		TSharedPtr<const FMeshData> ProcessedMesh = TaskHandle.GetTask()->GetMesh();
-		const FMeshData* ProcessedMeshPtr = ProcessedMesh.Get();
-
-		if (ensure(ProcessedMeshPtr != nullptr))
+		else
 		{
-			MeshPartitionEditorComponent->PostBuildSectionMesh(CompiledSection, *ProcessedMeshPtr, ModifierPtrs);
-
-			const int32 Vertices = ProcessedMesh->VertexCount();
-			const int32 Tris = ProcessedMesh->TriangleCount();
-
-			// TODO: this can be moved async (if we track intermediate state, and ensure GC doesn't reclaim our stuff)
-
-			MeshPartition::FTransformerUnit InitialUnit;
+			UE_LOG(LogMegaMeshEditor, Warning, TEXT("Grid-split build failed for compiled section %s — section left as placeholder."),
+				   *InCompiledSection->GetActorNameOrLabel());
+		}
 
-			InitialUnit.Section = CompiledSection;
-			InitialUnit.MeshData = ProcessedMesh;
-			InitialUnit.bShouldRecomputeNormals = true;
-			InitialUnit.bShouldRecomputeTangents = true;
+		return;
+	}
 
-			TUniquePtr<MeshPartition::FTransformerContext> TransformerContext = MeshPartitionEditorComponent->LaunchTransformers(MoveTemp(InitialUnit), Definition, BuildVariant);
-			
-			if (TransformerContext.IsValid() && TransformerContext->JoinTask.IsValid() && !TransformerContext->JoinTask.IsCompleted())
-			{
-				WaitOnGameThread(*TransformerContext);
-			}
+	// --- Non-grid path (bSplitSectionsToMatchWorldPartitionRuntimeGrid == false) ---
+	TSharedPtr<const FMeshData> ProcessedMesh = TaskHandle.GetTask()->GetMesh();
+	TSharedPtr<const MeshPartition::FSectionChannels> SectionChannel = TaskHandle.GetTask()->GetSectionChannels();
 
-			MeshPartitionEditorComponent->PostProcessSection(CompiledSection, ModifierPtrs);
-		
-			CompiledSection->SetIsPlaceholder(false);
+	if (ensure(ProcessedMesh.IsValid()) && ensure(SectionChannel.IsValid()))
+	{
+		MeshPartitionEditorComponent->ApplyBuiltMeshToSection(InCompiledSection, ProcessedMesh, *SectionChannel, ModifierGroup);
+		InCompiledSection->SetIsPlaceholder(false);
 
-			UE_LOG(LogMegaMeshEditor, Verbose, TEXT("Built Compiled Section: %p %s (%d bases, %d vertices, %d tris)"), CompiledSection, *CompiledSection->GetActorNameOrLabel(), BuildInfo.BaseModifierPaths.Num(), Vertices, Tris);
-		}
+		UE_LOG(LogMegaMeshEditor, Verbose, TEXT("Built Compiled Section: %p %s (%d bases, %d vertices, %d tris)"),
+			   InCompiledSection, *InCompiledSection->GetActorNameOrLabel(), BuildInfo.BaseModifierPaths.Num(),
+			   ProcessedMesh->VertexCount(), ProcessedMesh->TriangleCount());
 	}
 }
 
-void FSeparateWorldBuilder::RequestCompiledSectionBuild(MeshPartition::ACompiledSection* CompiledSection)
+void FSeparateWorldBuilder::RequestCompiledSectionBuild(MeshPartition::ACompiledSection* InCompiledSection)
 {
-	if (ensure(CompiledSection->IsPlaceholder()))
+	if (ensure(InCompiledSection->IsPlaceholder()))
 	{
 		check(IsInGameThread());
-		CompiledSectionsToBuild.Add(CompiledSection);
+		CompiledSectionsToBuild.Add(InCompiledSection);
 	}
 }
 
 void FSeparateWorldBuilder::Tick()
 {
 	check(IsInGameThread());
-	int32 SectionIndex = 0;
 	MeshPartition::EditorUtils::FPIEPathFixer& PIEPathFixer = UMeshPartitionEditorSubsystem::GetPIEPathFixer();
-	for (TWeakObjectPtr<MeshPartition::ACompiledSection> CompiledSectionPtr : CompiledSectionsToBuild)
+
+	// Swap into a local array before iterating. BuildCompiledSectionMesh pumps the game
+	// thread (WaitOnGameThread), which can trigger streaming → new RequestCompiledSectionBuild
+	// calls that Add to CompiledSectionsToBuild. Iterating the member directly would
+	// invalidate the range-for iterator on reallocation.
+	TArray<TWeakObjectPtr<MeshPartition::ACompiledSection>> SectionsToBuild;
+	Swap(SectionsToBuild, CompiledSectionsToBuild);
+
+	int32 SectionIndex = 0;
+	for (TWeakObjectPtr<MeshPartition::ACompiledSection> CompiledSectionPtr : SectionsToBuild)
 	{
-		UE_LOG(LogMegaMeshEditor, Verbose, TEXT("Building Placeholder Compiled Section (%d/%d) ..."), SectionIndex, CompiledSectionsToBuild.Num());
+		UE_LOG(LogMegaMeshEditor, Verbose, TEXT("Building Placeholder Compiled Section (%d/%d) ..."), SectionIndex, SectionsToBuild.Num());
 
 		if (MeshPartition::ACompiledSection* CompiledSection = CompiledSectionPtr.Get())
 		{
 			if (CompiledSection->IsPlaceholder())
 			{
 				FSoftObjectPath MeshPartitionPath = CompiledSection->GetBuildInfo().MegaMeshPath;
+
 				check(MeshPartitionPath.IsValid());
 
 				// MeshPartition path may be a PIE path -- we need to fix it to get the NonPIE package name
 				PIEPathFixer.FixInPlace(MeshPartitionPath);
-				
+
 				// TODO: we should ideally sort requests to avoid changing worlds as much as possible..
 				// this could be terrible if requests are coming from multiple worlds simultaneously.
 				// For now it should be ok as we generally work on one world at a time.
@@ -296,13 +509,21 @@ void FSeparateWorldBuilder::Tick()
 			}
 		}
 
-		UE_LOG(LogMegaMeshEditor, Verbose, TEXT("BUILT Placeholder Compiled Section (%d/%d)"), SectionIndex, CompiledSectionsToBuild.Num());
+		UE_LOG(LogMegaMeshEditor, Verbose, TEXT("BUILT Placeholder Compiled Section (%d/%d)"), SectionIndex, SectionsToBuild.Num());
 		SectionIndex++;
 	}
-	CompiledSectionsToBuild.Empty();
+
+	// Clear the build cache after processing all queued placeholders.
+	// This ensures cached data (including actor refs) is released promptly.
+	if (GroupBuildCache)
+	{
+		GroupBuildCache->Entries.Empty();
+	}
 }
 
 FSeparateWorldBuilder::~FSeparateWorldBuilder()
 {
+	// GroupBuildCache is destroyed automatically via TSharedPtr
 }
-} // namespace UE::MeshPartition
\ No newline at end of file
+
+} // namespace UE::MeshPartition
diff --git a/Engine/Plugins/Experimental/MeshPartition/Source/MeshPartitionEditor/Private/MeshPartitionWorldPartitionHelpers.cpp b/Engine/Plugins/Experimental/MeshPartition/Source/MeshPartitionEditor/Private/MeshPartitionWorldPartitionHelpers.cpp
--- a/Engine/Plugins/Experimental/MeshPartition/Source/MeshPartitionEditor/Private/MeshPartitionWorldPartitionHelpers.cpp
+++ b/Engine/Plugins/Experimental/MeshPartition/Source/MeshPartitionEditor/Private/MeshPartitionWorldPartitionHelpers.cpp
@@ -6,15 +6,20 @@
 #include "Engine/StaticMesh.h"
 #include "Engine/LevelStreaming.h"
 #include "MeshPartition.h"
-#include "MeshPartitionDefinition.h"
 #include "MeshPartitionCompiledSection.h"
+#include "MeshPartitionDefinition.h"
 #include "MeshPartitionEditorModule.h"
+#include "MeshPartitionMeshBuilder.h"
 #include "MeshPartitionModifierComponent.h"
 #include "MeshPartitionModifierDescriptors.h"
 #include "MeshPartitionModifierTaskGraph.h"
+#include "MeshPartitionTransformerPipeline.h"
+#include "MeshPartitionWPActorPropertiesTransformer.h"
 #include "Misc/App.h"
 #include "SourceControlHelpers.h"
 #include "UObject/SavePackage.h"
+#include "WorldPartition/RuntimeHashSet/RuntimePartitionLHGrid.h"
+#include "WorldPartition/RuntimeHashSet/WorldPartitionRuntimeHashSet.h"
 #include "WorldPartition/WorldPartition.h"
 #include "WorldPartition/WorldPartitionHelpers.h"
 #include "Algo/ForEach.h"
@@ -23,6 +28,117 @@
 
 namespace UE::MeshPartition::WorldPartitionHelpers
 {
+	uint32 GetGridCellSize(const UWorld* InWorld, FName InGridName)
+	{
+		if (InWorld == nullptr)
+		{
+			return 0;
+		}
+
+		// NAME_None is valid: ResolveRuntimePartition maps it to the first (main) grid.
+
+		const UWorldPartition* WorldPartition = InWorld->GetWorldPartition();
+		
+		if (WorldPartition == nullptr)
+		{
+			return 0;
+		}
+
+		const UWorldPartitionRuntimeHashSet* HashSet = Cast<UWorldPartitionRuntimeHashSet>(WorldPartition->RuntimeHash);
+		
+		if (HashSet == nullptr)
+		{
+			return 0;
+		}
+
+		const URuntimePartition* Partition = HashSet->ResolveRuntimePartition(InGridName, /*bMainPartitionLayer*/ true);
+		
+		if (const URuntimePartitionLHGrid* LHGrid = Cast<URuntimePartitionLHGrid>(Partition))
+		{
+			return LHGrid->GetCellSize();
+		}
+
+		return 0;
+	}
+
+	uint32 ResolveGridCellSizeFromPipeline(const MeshPartition::FCompiledSectionBuildVariant& InBuildVariant, const UWorld* InWorld)
+	{
+		if (InBuildVariant.TransformerPipeline == nullptr)
+		{
+			return 0;
+		}
+
+		// Find the last FWPActorPropertiesTransformer in the pipeline (iterate in reverse)
+		const TArray<TInstancedStruct<MeshPartition::FTransformer>>& Transformers = InBuildVariant.TransformerPipeline->GetTransformers();
+		
+		for (int32 TransformerIndex = Transformers.Num() - 1; TransformerIndex >= 0; --TransformerIndex)
+		{
+			if (Transformers[TransformerIndex].GetScriptStruct() == MeshPartition::FWPActorPropertiesTransformer::StaticStruct())
+			{
+				const FName GridName = Transformers[TransformerIndex].Get<MeshPartition::FWPActorPropertiesTransformer>().GetRuntimeGrid();
+				uint32 CellSize = GetGridCellSize(InWorld, GridName);
+				
+				if (CellSize > 0)
+				{
+					return CellSize;
+				}
+
+				// Fallback: if the named grid was not found, try the main grid
+				if (!GridName.IsNone())
+				{
+					CellSize = GetGridCellSize(InWorld, NAME_None);
+				
+					if (CellSize > 0)
+					{
+						UE_LOG(LogMegaMeshEditor, Warning, TEXT("WP grid '%s' not found — falling back to main grid (cell size %u)."),
+							   *GridName.ToString(), CellSize);
+						
+						return CellSize;
+					}
+				}
+
+				UE_LOG(LogMegaMeshEditor, Warning, TEXT("Could not resolve WP grid cell size for grid '%s' — skipping grid-aligned splitting."), *GridName.ToString());
+				
+				return 0;
+			}
+		}
+
+		return 0;
+	}
+
+	TArray<FGridCellEstimate> EstimateGridCells(const FBox& InGroupBounds, double InCellSize)
+	{
+		TArray<FGridCellEstimate> Cells;
+
+		if (InCellSize <= 0.0 || !InGroupBounds.IsValid)
+		{
+			return Cells;
+		}
+
+		const GridHelpers::FGridDimensions Grid = GridHelpers::ComputeGridDimensions(InGroupBounds, InCellSize);
+		Cells.Reserve(Grid.TotalCells);
+
+		UE_LOG(LogMegaMeshEditor, Verbose, TEXT("EstimateGridCells: Bounds [%.1f,%.1f,%.1f]-[%.1f,%.1f,%.1f] CellSize=%.1f → %dx%dx%d = %d cells"),
+			   InGroupBounds.Min.X, InGroupBounds.Min.Y, InGroupBounds.Min.Z,
+			   InGroupBounds.Max.X, InGroupBounds.Max.Y, InGroupBounds.Max.Z,
+			   InCellSize, Grid.CellNumber.X, Grid.CellNumber.Y, Grid.CellNumber.Z, Grid.TotalCells);
+
+		for (int32 SectionIndex = 0; SectionIndex < Grid.TotalCells; ++SectionIndex)
+		{
+			const int32 IndexX = SectionIndex % Grid.CellNumber.X;
+			const int32 IndexY = (SectionIndex / Grid.CellNumber.X) % Grid.CellNumber.Y;
+			const int32 IndexZ = SectionIndex / (Grid.CellNumber.X * Grid.CellNumber.Y);
+
+			const FIntVector CellCoord = Grid.OriginCoord + FIntVector(IndexX, IndexY, IndexZ);
+			const FVector CellMin = Grid.SnappedMin + FVector(IndexX, IndexY, IndexZ) * InCellSize;
+			const FVector CellMax = CellMin + FVector(InCellSize);
+
+			Cells.Add({ CellCoord, FBox(CellMin, CellMax) });
+		}
+
+		return Cells;
+	}
+
 	TSet<ULevelStreaming*> LoadAllLevelInstances(UWorldPartition* WorldPartition, TSet<FWorldPartitionReference>& InOutActorRefs)
 	{
 		TRACE_CPUPROFILER_EVENT_SCOPE(LoadAllLevelInstances);
@@ -85,7 +201,13 @@ namespace UE::MeshPartition::WorldPartitionHelpers
 		for (ULevelStreaming* LevelStreaming : ProcessedLevelStreamings)
 		{
 			ULevel* LoadedLevel = LevelStreaming->GetLoadedLevel();
-			check(LoadedLevel);
+			
+			if (LoadedLevel == nullptr)
+			{
+				// The level may not be loaded (e.g. it was streamed out between
+				// LoadAllLevelInstances and this iteration).
+				continue;
+			}
 
 			for (AActor* Actor : LoadedLevel->Actors)
 			{
diff --git a/Engine/Plugins/Experimental/MeshPartition/Source/MeshPartitionEditor/Private/MeshPartitionWorldUpdater.cpp b/Engine/Plugins/Experimental/MeshPartition/Source/MeshPartitionEditor/Private/MeshPartitionWorldUpdater.cpp
--- a/Engine/Plugins/Experimental/MeshPartition/Source/MeshPartitionEditor/Private/MeshPartitionWorldUpdater.cpp
+++ b/Engine/Plugins/Experimental/MeshPartition/Source/MeshPartitionEditor/Private/MeshPartitionWorldUpdater.cpp
@@ -10,6 +10,7 @@
 #include "Editor/EditorEngine.h"
 #include "Engine/Texture2DArray.h"
 #include "MeshPartitionEditorUtils.h"
+#include "MeshPartitionWorldPartitionHelpers.h"
 
 #define LOCTEXT_NAMESPACE "FMeshPartitionWorldUpdater"
 
@@ -19,6 +20,8 @@ namespace UE::MeshPartition
 {
 	FMeshPartitionWorldUpdater::FMeshPartitionWorldUpdater(UWorld* InWorld, EUpdateMode InUpdateMode)
 	{
+		TRACE_CPUPROFILER_EVENT_SCOPE(MeshPartition::FMeshPartitionWorldUpdater::FMeshPartitionWorldUpdater);
+
 		// create a unique build key to identify placeholders created
 		UpdateMode = InUpdateMode;
 		BuildKey = FGuid::NewGuid();
@@ -94,7 +97,7 @@ namespace UE::MeshPartition
 				// add an entry for this mesh partition
 				FMeshPartitionUpdater& MeshPartitionUpdater = MeshPartitionUpdaters.Add(
 					MeshPartitionGUID,
-					FMeshPartitionUpdater(BuildKey, *MPDefinition, MeshPartitionAllCachedDescriptors, TConstArrayView<FName>(BuildVariantNames), *AssetRegistry));
+					FMeshPartitionUpdater(BuildKey, *MPDefinition, MeshPartitionAllCachedDescriptors, TConstArrayView<FName>(BuildVariantNames), *AssetRegistry, InWorld));
 
 				// in PIE, we create placeholders for any missing bases.  In cook, we just warn that they are missing.
 				if (UpdateMode == EUpdateMode::ForCook)
@@ -115,9 +118,11 @@ namespace UE::MeshPartition
 		}
 	}
 
-	FMeshPartitionUpdater::FMeshPartitionUpdater(const FGuid& InBuildKey, const UMeshPartitionDefinition& InMeshPartitionDefinition, TSharedPtr<MeshPartition::FCachedDescriptors> InMeshPartitionAllCachedDescriptors, TConstArrayView<FName> InTargetBuildVariantNames, IAssetRegistry& InAssetRegistry)
+	FMeshPartitionUpdater::FMeshPartitionUpdater(const FGuid& InBuildKey, const UMeshPartitionDefinition& InMeshPartitionDefinition, TSharedPtr<MeshPartition::FCachedDescriptors> InMeshPartitionAllCachedDescriptors, TConstArrayView<FName> InTargetBuildVariantNames, IAssetRegistry& InAssetRegistry, const UWorld* InWorld)
 		: AllCachedDescriptors(InMeshPartitionAllCachedDescriptors)
 	{
+		TRACE_CPUPROFILER_EVENT_SCOPE(MeshPartition::FMeshPartitionUpdater::FMeshPartitionUpdater);
+
 		check(AllCachedDescriptors.IsValid());
 
 		// create an FBuildVariantUpdater for each of the target build variants in the mesh partition
@@ -131,7 +136,7 @@ namespace UE::MeshPartition
 			else
 			{
 				check(TargetGrouping->GetBuildVariantName() == TargetBuildVariantName);
-				BuildVariantUpdaters.Emplace(TargetBuildVariantName, FBuildVariantUpdater(InBuildKey, InMeshPartitionDefinition, *TargetGrouping, AllCachedDescriptors, InAssetRegistry));
+				BuildVariantUpdaters.Emplace(TargetBuildVariantName, FBuildVariantUpdater(InBuildKey, InMeshPartitionDefinition, *TargetGrouping, AllCachedDescriptors, InAssetRegistry, InWorld));
 			}
 		}
 
@@ -190,6 +195,7 @@ namespace UE::MeshPartition
 				continue;
 			}
 
+			Status.bBuildVariantHashMatches = (BuildInfo.BuildVariantHash == BuildVariantUpdater->GetBuildInfo().BuildVariantHash);
 			Status.bPackageHashMatches = UE::MeshPartition::EditorUtils::PackageHashIsUpToDate(BuildInfo, InAssetRegistry);
 			Status.bClassHashMatches = UE::MeshPartition::EditorUtils::ClassHashIsUpToDate(BuildInfo);
 
@@ -212,37 +218,41 @@ namespace UE::MeshPartition
 				continue;
 			}
 
-			const bool bUpToDate = bStillExists && Status.bPackageHashMatches && Status.bClassHashMatches;
+			const bool bUpToDate = bStillExists && Status.bBuildVariantHashMatches && Status.bPackageHashMatches && Status.bClassHashMatches;
 			if (bUpToDate)
 			{
-				// check if there is a duplicate compiled section for the same group
-				int32& BestCompiledSectionIndex = BuildVariantUpdater->BestReusableCompiledSectionIndexForGroup[GroupIndex];
-				if (BestCompiledSectionIndex < 0)
-				{
-					// no previous best compiled section, use this one
-					BestCompiledSectionIndex = CompiledSectionIndex;
-					Status.bIsDuplicate = false;
-				}
-				else
+				// Check for duplicates within the same grid cell (or same non-grid section).
+				// Grid-split sections are keyed by GridCellCoord; non-grid sections use InvalidGridCellCoord.
+				const FIntVector GridCellCoord = BuildInfo.GridCellCoord;
+				TMap<FIntVector, int32>& ReusableIndices = BuildVariantUpdater->ReusableCompiledSectionIndicesForGroup[GroupIndex];
+
+				if (int32* ExistingCompiledSectionIndex = ReusableIndices.Find(GridCellCoord))
 				{
-					// there is a previous best -- determine which one to use
-					MeshPartition::FCompiledSectionStatus& PrevBestCompiledSectionStatus = CompiledSectionStatus[BestCompiledSectionIndex];
-					check(PrevBestCompiledSectionStatus.bChecked && PrevBestCompiledSectionStatus.bReuse);
+					// there is a previous match for this grid cell -- determine which one to keep
+					MeshPartition::FCompiledSectionStatus& PrevBestStatus = CompiledSectionStatus[*ExistingCompiledSectionIndex];
+					check(PrevBestStatus.bChecked && PrevBestStatus.bReuse);
 
-					// we just use an arbitrary deterministic method to determine which one to keep
-					if (CompiledSections[BestCompiledSectionIndex].ActorDescGuid < CompiledSection.ActorDescGuid)
+					// use an arbitrary deterministic method to determine which one to keep
+					if (CompiledSections[*ExistingCompiledSectionIndex].ActorDescGuid < CompiledSection.ActorDescGuid)
 					{
-						// use the previous best -- this one is a duplicate (and should not be reused)
+						// use the previous best -- this one is a duplicate
 						Status.bIsDuplicate = true;
 					}
 					else
 					{
-						// use this one -- set previous best to be a duplicate
-						PrevBestCompiledSectionStatus.bIsDuplicate = true;
-						PrevBestCompiledSectionStatus.bReuse = false;
-						BestCompiledSectionIndex = CompiledSectionIndex;
+						// use this one -- set previous to be a duplicate
+						PrevBestStatus.bIsDuplicate = true;
+						PrevBestStatus.bReuse = false;
+						*ExistingCompiledSectionIndex = CompiledSectionIndex;
 					}
 				}
+				else
+				{
+					// no previous match for this grid cell, use this one
+					ReusableIndices.Add(GridCellCoord, CompiledSectionIndex);
+					Status.bIsDuplicate = false;
+				}
+
 				Status.bReuse = bUpToDate && !Status.bIsDuplicate;
 			}
 		}
@@ -251,14 +261,13 @@ namespace UE::MeshPartition
 		for (auto& Pair : BuildVariantUpdaters)
 		{
 			FBuildVariantUpdater& BuildVariantUpdater = Pair.Value;
-			for (int GroupIndex = 0; GroupIndex < BuildVariantUpdater.GetTargetModifierGroups().Num(); GroupIndex++)
+			for (int32 GroupIndex = 0; GroupIndex < BuildVariantUpdater.GetTargetModifierGroups().Num(); GroupIndex++)
 			{
-				bool bExistsReuse = true;
+				const TMap<FIntVector, int32>& ReusableIndices = BuildVariantUpdater.ReusableCompiledSectionIndicesForGroup[GroupIndex];
 
-				int32 CompiledSectionIndex = BuildVariantUpdater.BestReusableCompiledSectionIndexForGroup[GroupIndex];
-				if (CompiledSectionIndex < 0)
+				if (ReusableIndices.IsEmpty())
 				{
-					// this group does not have a reusable compiled section, and needs to be created.
+					// this group does not have any reusable compiled sections, and needs to be created.
 					// add all of the target group's base modifiers to the list of those needing to be created
 					const MeshPartition::FModifierGroup& ModifierGroup = BuildVariantUpdater.GetTargetModifierGroups()[GroupIndex];
 					for (const MeshPartition::FModifierDesc& ModifierDesc : ModifierGroup.BaseDescs())
@@ -268,18 +277,23 @@ namespace UE::MeshPartition
 				}
 				else
 				{
-					// sanity check
-					MeshPartition::FCompiledSectionStatus& Status = CompiledSectionStatus[CompiledSectionIndex];
-					check(Status.bChecked && Status.bReuse);
+					// sanity check all reusable sections
+					for (const auto& GridIndexPair : ReusableIndices)
+					{
+						MeshPartition::FCompiledSectionStatus& Status = CompiledSectionStatus[GridIndexPair.Value];
+						check(Status.bChecked && Status.bReuse);
+					}
 				}
 			}
 		}
 	}
 
-	FBuildVariantUpdater::FBuildVariantUpdater(	const FGuid& InBuildKey, const UMeshPartitionDefinition& InMeshPartitionDefinition, const MeshPartition::FModifierGrouping& InTargetGrouping, TSharedPtr<MeshPartition::FCachedDescriptors> InAllMeshPartitionDescriptors, IAssetRegistry& InAssetRegistry)
+	FBuildVariantUpdater::FBuildVariantUpdater(	const FGuid& InBuildKey, const UMeshPartitionDefinition& InMeshPartitionDefinition, const MeshPartition::FModifierGrouping& InTargetGrouping, TSharedPtr<MeshPartition::FCachedDescriptors> InAllMeshPartitionDescriptors, IAssetRegistry& InAssetRegistry, const UWorld* InWorld)
 		: AllMeshPartitionDescriptors(InAllMeshPartitionDescriptors)
 		, TargetGrouping(InTargetGrouping)
 	{
+		TRACE_CPUPROFILER_EVENT_SCOPE(MeshPartition::BuildVariantUpdater::MatchCompiledSections);
+
 		using namespace EditorUtils;
 
 		check(&TargetGrouping != nullptr);
@@ -292,10 +306,20 @@ namespace UE::MeshPartition
 			BuildInfo.SetMegaMeshDefinition(&InMeshPartitionDefinition);
 			BuildInfo.MegaMeshGUID = InAllMeshPartitionDescriptors->GetMeshPartitionGUID();
 			BuildInfo.BuildVariantName = InTargetGrouping.GetBuildVariantName();
-			{	// hash the build variant
-				FDependencyHash DepdendencyHash;
-				InMeshPartitionDefinition.GatherDependencies(DepdendencyHash, BuildInfo.BuildVariantName);
-				BuildInfo.BuildVariantHash = DepdendencyHash.GetDependentDataHash();
+			{	// hash the build variant — must match the main builder's hash (WorldPartitionMeshPartitionBuilder.cpp)
+				const MeshPartition::FCompiledSectionBuildVariant& BuildVariant = InMeshPartitionDefinition.GetCompiledSectionBuildVariantByName(BuildInfo.BuildVariantName);
+
+				FDependencyHash DependencyHash;
+				InMeshPartitionDefinition.GatherDependencies(DependencyHash, BuildInfo.BuildVariantName);
+
+				// Resolve and hash grid cell size to match the main builder
+				if (BuildVariant.bSplitSectionsToMatchWorldPartitionRuntimeGrid)
+				{
+					ResolvedGridCellSize = WorldPartitionHelpers::ResolveGridCellSizeFromPipeline(BuildVariant, InWorld);
+					DependencyHash += ResolvedGridCellSize;
+				}
+
+				BuildInfo.BuildVariantHash = DependencyHash.GetDependentDataHash();
 			}
 			// Note: get the path from the descriptor ref, because it will be the original path (not the PIE temporary path)
 			BuildInfo.MegaMeshPath = InAllMeshPartitionDescriptors->GetMeshPartitionPath();
@@ -333,8 +357,8 @@ namespace UE::MeshPartition
 
 		TConstArrayView<MeshPartition::FModifierGroup> TargetModifierGroups = InTargetGrouping.GetModifierGroups();
 
-		// init all groups to have no corresponding compiled section (we will determine the mapping later in the MeshPartition updater)
-		BestReusableCompiledSectionIndexForGroup.Init(-1, TargetModifierGroups.Num());
+		// init all groups to have no corresponding compiled sections (we will determine the mapping later in the MeshPartition updater)
+		ReusableCompiledSectionIndicesForGroup.SetNum(TargetModifierGroups.Num());
 
 		// build the ModifierSetHashToTargetGroupIndex
 		for (int GroupIndex = 0; GroupIndex < TargetModifierGroups.Num(); GroupIndex++)
@@ -346,21 +370,26 @@ namespace UE::MeshPartition
 	}
 
 
-	int32 FBuildVariantUpdater::GetReusableCompiledSectionForGroup(int32 GroupIndex, const FMeshPartitionUpdater& MeshPartitionUpdater) const
+	bool FBuildVariantUpdater::HasReusableCompiledSectionsForGroup(int32 InGroupIndex, const FMeshPartitionUpdater& InMeshPartitionUpdater) const
 	{
-		int32 CompiledSectionIndex = BestReusableCompiledSectionIndexForGroup[GroupIndex];
-		if (const FCompiledSectionStatus* Status = MeshPartitionUpdater.GetCompiledSectionStatus(CompiledSectionIndex))
+		const TMap<FIntVector, int32>& ReusableIndices = ReusableCompiledSectionIndicesForGroup[InGroupIndex];
+
+		for (const auto& Pair : ReusableIndices)
 		{
-			ensure(Status->bChecked); // should have been classified by now
-			ensure(Status->bReuse);
-			// this group has a reusable compiled section
-			return CompiledSectionIndex;
+			if (const FCompiledSectionStatus* Status = InMeshPartitionUpdater.GetCompiledSectionStatus(Pair.Value))
+			{
+				ensure(Status->bChecked);
+				ensure(Status->bReuse);
+			}
 		}
-		return -1;
+
+		return !ReusableIndices.IsEmpty();
 	}
 
 	int32 FMeshPartitionWorldUpdater::CreatePlaceholderActors(UWorld* InTargetWorld)
 	{
+		TRACE_CPUPROFILER_EVENT_SCOPE(MeshPartition::WorldUpdater::CreatePlaceholderActors);
+
 		check(UpdateMode == EUpdateMode::ForPIE);
 
 		int32 PlaceholdersCreated = 0;
@@ -372,62 +401,83 @@ namespace UE::MeshPartition
 				TConstArrayView<MeshPartition::FModifierGroup> ModifierGroups = BuildVariantUpdater.GetTargetGrouping().GetModifierGroups();
 				for (int32 GroupIndex = 0; GroupIndex < ModifierGroups.Num(); GroupIndex++)
 				{
-					int32 CompiledSectionIndex = BuildVariantUpdater.GetReusableCompiledSectionForGroup(GroupIndex, MeshPartitionUpdater);
-					if (CompiledSectionIndex >= 0)
+					if (BuildVariantUpdater.HasReusableCompiledSectionsForGroup(GroupIndex, MeshPartitionUpdater))
 					{
-						continue; // no need to create a placeholder compiled section if there this group already has a reuseable one
+						continue; // no need to create placeholders if this group already has reusable sections
 					}
 
-					// insert a placeholder for this modifier group!
+					// insert placeholder(s) for this modifier group
 					const MeshPartition::FModifierGroup& ModifierGroup = ModifierGroups[GroupIndex];
 					const MeshPartition::UMeshPartitionDefinition* Definition = BuildVariantUpdater.GetBuildInfo().GetMegaMeshDefinition();
 					const MeshPartition::FCompiledSectionBuildVariant& BuildVariant = Definition->GetCompiledSectionBuildVariantByName(BuildVariantUpdater.GetBuildVariantName());
 
 					// Build the bases list from the modifier paths
-					FCompiledSectionBuildInfo BuildInfo = BuildVariantUpdater.GetBuildInfo();
-					BuildInfo.BaseModifierPaths.Empty();
-					Algo::Transform(ModifierGroup.BaseDescs(), BuildInfo.BaseModifierPaths, [](const MeshPartition::FModifierDesc& BaseDescriptor)
+					FCompiledSectionBuildInfo BaseBuildInfo = BuildVariantUpdater.GetBuildInfo();
+					BaseBuildInfo.BaseModifierPaths.Empty();
+					Algo::Transform(ModifierGroup.BaseDescs(), BaseBuildInfo.BaseModifierPaths, [](const MeshPartition::FModifierDesc& InBaseDescriptor)
 						{
 							// the path is taken from the editor world, so it should not need any massaging (unlike PIE world paths)
-							return BaseDescriptor.ModifierPath;
+							return InBaseDescriptor.ModifierPath;
 						});
 
-					// Note: We don't fill out ModifiersSetHash because we don't have the full set of non-base modifiers
-					// (The placeholder builder will find them when it loads everything)
-					FString BaseName = TEXT("CompiledSection_") + BuildVariant.Name.ToString() + TEXT("_PIE");
-					FName UniqueName = MakeUniqueObjectName(InTargetWorld, UStaticMesh::StaticClass(), FName(BaseName), EUniqueObjectNameOptions::UniversallyUnique);
-
-					FActorSpawnParameters Params;
-					Params.Name = UniqueName;
-					// Params.Owner = MeshPartitionActor; // TODO: we could parent to the main actor, but we would have to search for it first
-					Params.SpawnCollisionHandlingOverride = ESpawnActorCollisionHandlingMethod::AlwaysSpawn;
-					Params.bNoFail = true;
-					Params.ObjectFlags = RF_NoFlags;		// non transational (does not interact with undo/redo)
-
-					MeshPartition::ACompiledSection* Placeholder = InTargetWorld->SpawnActor<MeshPartition::ACompiledSection>(MeshPartition::ACompiledSection::StaticClass(), FTransform::Identity, Params);
-					if (Placeholder->GetRootComponent())
+					// Compute bounds from base modifiers + growth-flagged non-bases.
+					// ComputeBaseBounds unions BaseDescs bounds, then extends along axes
+					// where non-base modifiers have BaseGrowth set (actual deformations).
+					// This avoids inflating bounds from global modifiers with huge extents.
+					const FBox GroupBounds = ModifierGroup.ComputeBaseBounds();
+
+					// Determine placeholder grid: either N per-cell placeholders or one per group
+					TArray<WorldPartitionHelpers::FGridCellEstimate> PlaceholderCells;
+					const uint32 GridCellSize = BuildVariantUpdater.GetResolvedGridCellSize();
+
+					if (BuildVariant.bSplitSectionsToMatchWorldPartitionRuntimeGrid && (GridCellSize > 0))
 					{
-						Placeholder->GetRootComponent()->SetMobility(EComponentMobility::Static);
+						// Grid-split: create one placeholder per estimated grid cell
+						PlaceholderCells = WorldPartitionHelpers::EstimateGridCells(GroupBounds, static_cast<double>(GridCellSize));
 					}
-					Placeholder->SetBuildInfo(BuildInfo);
-					Placeholder->SetIsPlaceholder(true);
-					Placeholder->SetActorLabel(BaseName);
+					else
+					{
+						// Non-grid: single placeholder with full group bounds
+						PlaceholderCells.Add({ FCompiledSectionBuildInfo::InvalidGridCellCoord, GroupBounds });
+					}
+
+					// Note: We don't fill out ModifiersSetHash because we don't have the full set of non-base modifiers
+					// (The placeholder builder will find them when it loads everything)
+					const FString BaseName = TEXT("CompiledSection_") + BuildVariant.Name.ToString() + TEXT("_PIE");
 
-					// build the estimated streaming bounds by unioning each base descriptor's bounds
-					FBox EstimatedStreamingBounds(ForceInit);
+					for (const WorldPartitionHelpers::FGridCellEstimate& Cell : PlaceholderCells)
 					{
-						for (const FModifierDesc& Desc : ModifierGroup.BaseDescs())
+						FName UniqueName = MakeUniqueObjectName(InTargetWorld, UStaticMesh::StaticClass(), FName(BaseName), EUniqueObjectNameOptions::UniversallyUnique);
+
+						FActorSpawnParameters SpawnParams;
+						SpawnParams.Name = UniqueName;
+						SpawnParams.SpawnCollisionHandlingOverride = ESpawnActorCollisionHandlingMethod::AlwaysSpawn;
+						SpawnParams.bNoFail = true;
+						SpawnParams.ObjectFlags = RF_NoFlags;
+
+						MeshPartition::ACompiledSection* Placeholder = InTargetWorld->SpawnActor<MeshPartition::ACompiledSection>(MeshPartition::ACompiledSection::StaticClass(), FTransform::Identity, SpawnParams);
+
+						if (Placeholder->GetRootComponent())
 						{
-							EstimatedStreamingBounds += Desc.Bounds;
+							Placeholder->GetRootComponent()->SetMobility(EComponentMobility::Static);
 						}
-						// apply some wiggle room for non-base modifiers -- this can be wildly off though..
-						EstimatedStreamingBounds = EstimatedStreamingBounds.ExpandBy(EstimatedStreamingBounds.GetExtent() * 0.1);
-					}
-					Placeholder->SetPlaceholderStreamingBounds(EstimatedStreamingBounds);
 
-					UE_LOG(LogMegaMeshEditor, Verbose, TEXT("    ADDED Placeholder %.6s (%p %s) (for %d bases)"), *Placeholder->GetActorGuid().ToString(), Placeholder, *Placeholder->GetActorNameOrLabel(), Placeholder->GetBuildInfo().BaseModifierPaths.Num());
-					PlaceholdersCreated++;
-					MeshPartitionUpdater.AddedPlaceholderCompiledSections++;
+						FCompiledSectionBuildInfo PlaceholderBuildInfo = BaseBuildInfo;
+						PlaceholderBuildInfo.GridCellCoord = Cell.GridCellCoord;
+						PlaceholderBuildInfo.ResolvedGridCellSize = GridCellSize;
+
+						Placeholder->SetBuildInfo(PlaceholderBuildInfo);
+						Placeholder->SetIsPlaceholder(true);
+						Placeholder->SetActorLabel(BaseName);
+						Placeholder->SetPlaceholderStreamingBounds(Cell.CellBounds);
+
+						UE_LOG(LogMegaMeshEditor, Display, TEXT("    ADDED Placeholder %.6s (%p %s) GridCell:(%d,%d,%d) ResolvedGridCellSize:%u (for %d bases)"),
+							   *Placeholder->GetActorGuid().ToString(), Placeholder, *Placeholder->GetActorNameOrLabel(),
+							   Cell.GridCellCoord.X, Cell.GridCellCoord.Y, Cell.GridCellCoord.Z, GridCellSize, Placeholder->GetBuildInfo().BaseModifierPaths.Num());
+
+						PlaceholdersCreated++;
+						MeshPartitionUpdater.AddedPlaceholderCompiledSections++;
+					}
 				}
 				return true; // continue iterating build variants
 			});
@@ -438,6 +488,8 @@ namespace UE::MeshPartition
 
 	int32 FMeshPartitionWorldUpdater::RemovePlaceholderActors()
 	{
+		TRACE_CPUPROFILER_EVENT_SCOPE(MeshPartition::WorldUpdater::RemovePlaceholderActors);
+
 		check(UpdateMode == EUpdateMode::ForPIE);
 
 		int32 PlaceholdersRemoved = 0;
@@ -554,6 +606,10 @@ namespace UE::MeshPartition
 			{
 				InOutSummary.CompiledSections.MismatchedModifiers++;
 			}
+			else if (!Status.bBuildVariantHashMatches)
+			{
+				InOutSummary.CompiledSections.BuildVariantHashFails++;
+			}
 			else if (!Status.bPackageHashMatches)
 			{
 				InOutSummary.CompiledSections.PackageHashFails++;
@@ -587,10 +643,10 @@ namespace UE::MeshPartition
 			InOutSummary.TargetBuildVariantCount++;
 
 			// look for any target groups with no corresponding compiled section -- these are missing
-			for (int32 ReuseCompiledSectionIndex : VariantUpdater.BestReusableCompiledSectionIndexForGroup)
+			for (const TMap<FIntVector, int32>& GridCellMap : VariantUpdater.ReusableCompiledSectionIndicesForGroup)
 			{
 				InOutSummary.GroupCount++;
-				if (ReuseCompiledSectionIndex < 0)
+				if (GridCellMap.IsEmpty())
 				{
 					InOutSummary.CompiledSections.Missing++;
 				}
@@ -727,16 +783,27 @@ namespace UE::MeshPartition
 
 				int32 TotalBaseCount = MeshPartitionUpdater.AllCachedDescriptors->GetAllBaseModifiers().Num();
 
-				// first report on the groups that reuse a compiled section
+				// first report on the groups that reuse a compiled section (one line per grid cell)
 				TConstArrayView<MeshPartition::FModifierGroup> ModifierGroups = VariantUpdater.GetTargetModifierGroups();
 				for (int32 GroupIndex = 0; GroupIndex < ModifierGroups.Num(); GroupIndex++)
 				{
 					const MeshPartition::FModifierGroup& Group = ModifierGroups[GroupIndex];
-					int32 CompiledSectionIndex = VariantUpdater.GetReusableCompiledSectionForGroup(GroupIndex, MeshPartitionUpdater);
-					ReportCompiledSection(CompiledSectionIndex, Group.BaseDescs().Num(), TotalBaseCount, BuildVariantName);
-					if (CompiledSectionIndex >= 0)
+					const TMap<FIntVector, int32>& GridCellMap = VariantUpdater.GetReusableCompiledSectionIndicesForGroup(GroupIndex);
+					if (GridCellMap.IsEmpty())
 					{
-						CompiledSectionIndicesReported.Add(CompiledSectionIndex);
+						// no compiled section exists for this group at all
+						ReportCompiledSection(-1, Group.BaseDescs().Num(), TotalBaseCount, BuildVariantName);
+					}
+					else
+					{
+						for (const TPair<FIntVector, int32>& CellPair : GridCellMap)
+						{
+							ReportCompiledSection(CellPair.Value, Group.BaseDescs().Num(), TotalBaseCount, BuildVariantName);
+							if (CellPair.Value >= 0)
+							{
+								CompiledSectionIndicesReported.Add(CellPair.Value);
+							}
+						}
 					}
 				}
 			}
@@ -758,8 +825,8 @@ namespace UE::MeshPartition
 
 			UE_LOG(LogMegaMeshEditor, Warning, TEXT(" ==="));
 
-			UE_LOG(LogMegaMeshEditor, Warning, TEXT(" - Target Groups: %d UpToDate, %d OutOfDate, %d Missing, %d Total"),
-				Summary.GroupsUpToDate(), Summary.GroupsOutOfDate(), Summary.GroupsMissing(), Summary.GroupCount);
+			UE_LOG(LogMegaMeshEditor, Warning, TEXT(" - Sections: %d UpToDate, %d OutOfDate, %d GroupsMissing, %d GroupsTotal"),
+				Summary.SectionsUpToDate(), Summary.SectionsOutOfDate(), Summary.GroupsMissing(), Summary.GroupCount);
 
 			UE_LOG(LogMegaMeshEditor, Warning, TEXT(" - CompiledSections: %d ToRemove, %d Total"),
 				Summary.CompiledSectionsToRemove(), Summary.CompiledSectionCount);
@@ -797,11 +864,20 @@ namespace UE::MeshPartition
 		{
 			StatusString += TEXT("MismatchedModifiers");
 		}
-		if (!bPackageHashMatches)
+
+		// Hash fields are only meaningful for sections that reached the hash-check stage
+		// (i.e., not early-categorized as placeholder, non-target variant, or mismatched modifiers)
+		const bool bReachedHashCheck = !bIsPlaceholder && !bNonTargetVariant && !bMismatchedModifiers;
+
+		if (bReachedHashCheck && !bBuildVariantHashMatches)
+		{
+			StatusString += TEXT("BuildVariantHashFail");
+		}
+		if (bReachedHashCheck && !bPackageHashMatches)
 		{
 			StatusString += TEXT("PackageHashFail");
 		}
-		if (!bClassHashMatches)
+		if (bReachedHashCheck && !bClassHashMatches)
 		{
 			StatusString += TEXT("ClassHashFail");
 		}
diff --git a/Engine/Plugins/Experimental/MeshPartition/Source/MeshPartitionEditor/Private/WorldPartitionMeshPartitionBuilder.cpp b/Engine/Plugins/Experimental/MeshPartition/Source/MeshPartitionEditor/Private/WorldPartitionMeshPartitionBuilder.cpp
--- a/Engine/Plugins/Experimental/MeshPartition/Source/MeshPartitionEditor/Private/WorldPartitionMeshPartitionBuilder.cpp
+++ b/Engine/Plugins/Experimental/MeshPartition/Source/MeshPartitionEditor/Private/WorldPartitionMeshPartitionBuilder.cpp
@@ -23,6 +23,7 @@
 #include "Misc/ScopedSlowTask.h"
 #include "ObjectTools.h"
 #include "SourceControlHelpers.h"
+#include "SourceControlOperations.h"
 #include "UObject/SavePackage.h"
 #include "WorldPartition/WorldPartition.h"
 #include "WorldPartition/WorldPartitionHelpers.h"
@@ -47,7 +48,6 @@
 #include "WorldPartition/DataLayer/DataLayerAsset.h"
 #include "WorldPartition/DataLayer/DataLayerManager.h"
 #include "WorldPartition/DataLayer/DataLayerInstance.h"
-#include "SourceControlOperations.h"
 
 #define LOCTEXT_NAMESPACE "MeshPartitionBuilder"
 
@@ -136,20 +136,20 @@ bool UWorldPartitionMeshPartitionBuilder::RunInternal(UWorld* InWorld, const FCe
 
 	UE_LOG(LogMegaMeshEditor, Display, TEXT("[%i / %i] Initializing compiled sections build..."), 0, TotalNumberOfSectionsToBuild);
 	
+	int32 TotalFailedMeshPartitions = 0;
+
 	// For each MeshPartition in the world
-	int32 FailedMeshPartitionBuilds = 0;
 	for (const FGuid& MeshPartitionGUID  : AllMeshPartitionGuidsInWorld)
 	{
 		// Use FWorldPartitionReference to ensure the AMeshPartition actor is loaded, then invoke BuildMeshPartition on it
 		FWorldPartitionReference MeshPartitionActorRef(WorldPartition.Get(), MeshPartitionGUID);
-
+		const AMeshPartition* MeshPartition = Cast<AMeshPartition>(MeshPartitionActorRef.GetActor());
+		
 		// Note: Not every MeshPartition GUID references a valid mesh partition actor -- for example, there may be stale compiled sections that reference the GUID of old mesh partitions that have since been deleted
-		if (const AMeshPartition* MeshPartition = Cast<AMeshPartition>(MeshPartitionActorRef.GetActor()))
+		if ((MeshPartition != nullptr) && !BuildMeshPartition(MeshPartition, BuildKey, *DescriptorCache, InPackageHelper))
 		{
-			if (!BuildMeshPartition(MeshPartition, BuildKey, *DescriptorCache, InPackageHelper))
-			{
-				FailedMeshPartitionBuilds++;
-			}
+			UE_LOG(LogMegaMeshEditor, Error, TEXT("BuildMeshPartition failed for '%s' (%.6s)"), *MeshPartition->GetActorNameOrLabel(), *MeshPartitionGUID.ToString());
+			TotalFailedMeshPartitions++;
 		}
 	}
 
@@ -215,7 +215,7 @@ bool UWorldPartitionMeshPartitionBuilder::RunInternal(UWorld* InWorld, const FCe
 							{
 								if (SCCState->IsCheckedOutOther())
 								{
-									// don't attempt to delete if it 
+									// don't attempt to delete if it is exclusively checked out by someone else
 									UE_LOG(LogMegaMeshEditor, Display,
 										TEXT("  Skipping deletion of compiled section %.6s: package '%s' is exclusively checked out by another user (%s)."),
 										*ActorDescGuid.ToString(), *PackageFilename, *SCCState->GetDisplayName().ToString());
@@ -271,71 +271,52 @@ bool UWorldPartitionMeshPartitionBuilder::RunInternal(UWorld* InWorld, const FCe
 		}
 	}
 
-	UE_LOG(LogMegaMeshEditor, Display, TEXT("Finished Build for World %s (Built %d, Reused %d, Failed %d, Removed %d, Could Not Remove %d Compiled Sections)"), *InWorld->GetName(),
-		TotalBuiltCompiledSections,
-		TotalReusedCompiledSections,
-		TotalFailedCompiledSections,
-		TotalRemovedCompiledSections,
-		TotalFailedToRemoveCompiledSections);
-
-	if ((FailedMeshPartitionBuilds == 0) && (TotalFailedCompiledSections == 0))
-	{
-		return true;
-	}
-
-	UE_LOG(LogMegaMeshEditor, Warning, TEXT("Failed to build for World %s (Built %d, Reused %d, Failed %d, Removed %d, Could Not Remove %d Compiled Sections), %d Mesh Partitions Failed"), *InWorld->GetName(),
+	UE_LOG(LogMegaMeshEditor, Display, TEXT("Finished Build for World %s (Built %d, Reused %d, Failed %d, Removed %d, Could Not Remove %d Compiled Sections, %d Failed Mesh Partitions)"), *InWorld->GetName(),
 		TotalBuiltCompiledSections,
 		TotalReusedCompiledSections,
 		TotalFailedCompiledSections,
 		TotalRemovedCompiledSections,
 		TotalFailedToRemoveCompiledSections,
-		FailedMeshPartitionBuilds);
+		TotalFailedMeshPartitions);
 
-	return false;
+	return (TotalFailedCompiledSections == 0) && (TotalFailedMeshPartitions == 0);
 }
 
-bool UWorldPartitionMeshPartitionBuilder::BuildMeshPartition(const AMeshPartition* MeshPartition, const FGuid& BuildKey, MeshPartition::FModifierDescriptorCache& DescriptorCache, FPackageSourceControlHelper& InPackageHelper)
+bool UWorldPartitionMeshPartitionBuilder::BuildMeshPartition(const AMeshPartition* InMeshPartition, const FGuid& BuildKey, MeshPartition::FModifierDescriptorCache& DescriptorCache, FPackageSourceControlHelper& InPackageHelper)
 {
 	TRACE_CPUPROFILER_EVENT_SCOPE(UWorldPartitionMeshPartitionBuilder::BuildMeshPartition);
 
-	if (!ensure(MeshPartition))
-	{
-		UE_LOG(LogMegaMeshEditor, Error, TEXT("Cannot build null MeshPartition actor reference."));
-		return false;
-	}
-
-	const FGuid& MeshPartitionGUID = MeshPartition->GetActorGuid();
-	if (!ensure(MeshPartitionGUID.IsValid()))
+	const FGuid& MeshPartitionGUID = InMeshPartition->GetActorGuid();
+	if (!MeshPartitionGUID.IsValid())
 	{
-		UE_LOG(LogMegaMeshEditor, Error, TEXT("Invalid MeshPartitionGUID for Mesh Partition '%s'"), *MeshPartition->GetActorNameOrLabel());
 		return false;
 	}
 
-	UWorld* World = MeshPartition->GetWorld();
+	UWorld* World = InMeshPartition->GetWorld();
 	check(World);
 
 	// Create a new build info -- this gets recorded on all compiled sections produced here.
 	MeshPartition::FCompiledSectionBuildInfo BuildInfo;
 	BuildInfo.BuildKey = BuildKey;
 	BuildInfo.MegaMeshGUID = MeshPartitionGUID;
-	BuildInfo.MegaMeshPath = MeshPartition;
+	BuildInfo.MegaMeshPath = InMeshPartition;
 
-	TSharedPtr<const MeshPartition::FCachedDescriptors> Descriptors = DescriptorCache.GetCachedDescriptors(MeshPartition);
+	TSharedPtr<const MeshPartition::FCachedDescriptors> Descriptors = DescriptorCache.GetCachedDescriptors(InMeshPartition);
 	if (!Descriptors.IsValid())
 	{
-		UE_LOG(LogMegaMeshEditor, Warning, TEXT("Could not gather descriptors for the Mesh Partition '%s'"), *MeshPartition->GetActorNameOrLabel());
+		UE_LOG(LogMegaMeshEditor, Warning, TEXT("Could not gather descriptors for the Mesh Partition '%s'"), *InMeshPartition->GetActorNameOrLabel());
 		return false;
 	}
 	TArrayView<const FCompiledSectionDescriptor> OldCompiledSections = Descriptors->GetCompiledSections();
 
 	// Use FWorldPartitionReference to load/unload everything in scope
 	FWorldPartitionReference MeshPartitionActorRef(WorldPartition.Get(), MeshPartitionGUID);
-	check(MeshPartition == CastChecked<AMeshPartition>(MeshPartitionActorRef.GetActor()));
+	check(InMeshPartition == CastChecked<AMeshPartition>(MeshPartitionActorRef.GetActor()));
 
-	const UMeshPartitionDefinition* MeshPartitionDefinition = MeshPartition->GetMeshPartitionDefinition();
+	const UMeshPartitionDefinition* MeshPartitionDefinition = InMeshPartition->GetMeshPartitionDefinition();
 	if (MeshPartitionDefinition == nullptr)
 	{
-		UE_LOG(LogMegaMeshEditor, Verbose, TEXT("No Mesh Partition Definition provided for %s (%.6s), using default settings."), *MeshPartition->GetActorNameOrLabel(), *MeshPartitionGUID.ToString());
+		UE_LOG(LogMegaMeshEditor, Verbose, TEXT("No Mesh Partition Definition provided for %s (%.6s), using default settings."), *InMeshPartition->GetActorNameOrLabel(), *MeshPartitionGUID.ToString());
 		MeshPartitionDefinition = UMeshPartitionDefinition::StaticClass()->GetDefaultObject<UMeshPartitionDefinition>();
 	}
 	BuildInfo.SetMegaMeshDefinition(MeshPartitionDefinition);
@@ -343,16 +324,15 @@ bool UWorldPartitionMeshPartitionBuilder::BuildMeshPartition(const AMeshPartitio
 	TConstArrayView<MeshPartition::FCompiledSectionBuildVariant> BuildVariants = MeshPartitionDefinition->GetCompiledSectionBuildVariants();
 	check(BuildVariants.Num() > 0);
 
-	UMeshPartitionEditorComponent* MeshPartitionEditorComponent = Cast<UMeshPartitionEditorComponent>(MeshPartition->GetMeshPartitionComponent());
+	UMeshPartitionEditorComponent* MeshPartitionEditorComponent = Cast<UMeshPartitionEditorComponent>(InMeshPartition->GetMeshPartitionComponent());
 	check(MeshPartitionEditorComponent != nullptr);
-
 	// avoid level streaming triggering a preview section build while we compiled section resources are being built.
 	MeshPartitionEditorComponent->SetPreviewSectionBuildEnabled(false);
 
-	const FTransform& Transform = MeshPartition->GetTransform();
+	const FTransform& Transform = InMeshPartition->GetTransform();
 
 	UE_LOG(LogMegaMeshEditor, Display, TEXT("-- %.6s Build Mesh Partition %s (%d Variants, %d Existing Compiled Sections"),
-		*MeshPartitionGUID.ToString(), *MeshPartition->GetActorNameOrLabel(),
+		*MeshPartitionGUID.ToString(), *InMeshPartition->GetActorNameOrLabel(),
 		BuildVariants.Num(), OldCompiledSections.Num());
 
 	if (AMeshPartitionDataLayerContainer* Container = AMeshPartitionDataLayerContainer::GetOrCreate(World))
@@ -381,6 +361,16 @@ bool UWorldPartitionMeshPartitionBuilder::BuildMeshPartition(const AMeshPartitio
 
 		FDependencyHash DependencyHash;
 		BuildInfo.GetMegaMeshDefinition()->GatherDependencies(DependencyHash, BuildVariant.Name);
+
+		// Resolve grid cell size from the pipeline's WP transformer and hash it into the build variant
+		uint32 ResolvedGridCellSize = 0;
+		
+		if (BuildVariant.bSplitSectionsToMatchWorldPartitionRuntimeGrid)
+		{
+			ResolvedGridCellSize = WorldPartitionHelpers::ResolveGridCellSizeFromPipeline(BuildVariant, World);
+			DependencyHash += ResolvedGridCellSize;
+		}
+
 		BuildInfo.BuildVariantHash = DependencyHash.GetDependentDataHash();
 
 		UE_LOG(LogMegaMeshEditor, Display, TEXT("-- %.6s -- %s Build Variant"), *MeshPartitionGUID.ToString(), *BuildVariant.Name.ToString());
@@ -410,11 +400,12 @@ bool UWorldPartitionMeshPartitionBuilder::BuildMeshPartition(const AMeshPartitio
 			Algo::Transform(ModifierGroup.BaseDescs(), BuildInfo.BaseModifierPaths, [](const MeshPartition::FModifierDesc& BaseDescriptor) { return BaseDescriptor.ModifierPath; });
 
 			FModifierGroupTask* ModifierGroupTask = new FModifierGroupTask(ModifierGroup, Settings, BuildInfo, BuildVariant, MeshPartitionEditorComponent, *this, OldCompiledSections);
+			ModifierGroupTask->ResolvedGridCellSize = ResolvedGridCellSize;
 			ModifierGroupTasks.Add(ModifierGroupTask);
 			ModifierGroupTask->State = ETaskState::Ready;
 		}
 
-		ExecuteModifierGroupTasks(ModifierGroupTasks, World, MeshPartition, InPackageHelper);
+		ExecuteModifierGroupTasks(ModifierGroupTasks, World, InMeshPartition, InPackageHelper);
 
 		int32 BuiltCompiledSections = 0;
 		int32 FailedCompiledSections = 0;
@@ -428,8 +419,12 @@ bool UWorldPartitionMeshPartitionBuilder::BuildMeshPartition(const AMeshPartitio
 				break;
 			case ETaskResult::ReuseExisting:
 				ReusedCompiledSections++;
-				check(ModifierGroupTask.ReuseCompiledSectionActorDescGuid.IsValid());
-				CompiledSectionActorDescGuidsToKeep.Add(ModifierGroupTask.ReuseCompiledSectionActorDescGuid);
+				check(!ModifierGroupTask.ReuseCompiledSectionActorDescGuids.IsEmpty());
+
+				for (const FGuid& Guid : ModifierGroupTask.ReuseCompiledSectionActorDescGuids)
+				{
+					CompiledSectionActorDescGuidsToKeep.Add(Guid);
+				}
 				break;
 			default:
 				FailedCompiledSections++;
@@ -455,7 +450,7 @@ bool UWorldPartitionMeshPartitionBuilder::BuildMeshPartition(const AMeshPartitio
 		FWorldPartitionHelpers::DoCollectGarbage();
 	}
 
-	UE_LOG(LogMegaMeshEditor, Display, TEXT("-- %.6s -- Finished Mesh Partition Build (%s)"), *MeshPartitionGUID.ToString(), *MeshPartition->GetActorNameOrLabel());
+	UE_LOG(LogMegaMeshEditor, Display, TEXT("-- %.6s -- Finished Mesh Partition Build (%s)"), *MeshPartitionGUID.ToString(), *InMeshPartition->GetActorNameOrLabel());
 	return true;
 }
 
@@ -560,16 +555,20 @@ EEditorBuildResult UWorldPartitionMeshPartitionBuilder::BuildMeshPartitions(UWor
 
 			if (LogReader != nullptr)
 			{
+				// Cap the read to avoid int64->int32 narrowing on TArray size. We only need recent
+				// progress lines, so reading the tail of the log is sufficient.
+				constexpr int64 MaxReadSize = 64 * 1024 * 1024; // 64 MB
+				const int64 BytesAvailable = FileStat.FileSize - LastReadPos;
+				const int64 ReadStart = (BytesAvailable > MaxReadSize) ? (FileStat.FileSize - MaxReadSize) : LastReadPos;
+				const int32 BytesToRead = static_cast<int32>(FMath::Min(BytesAvailable, MaxReadSize));
+
 				TArray<uint8> Buffer;
 
-				LogReader->Seek(LastReadPos);
-				int64 BufferSizeBytes = FileStat.FileSize - LastReadPos;
-				constexpr int64 MaxBufferToProcess = 256*1024*1024;
-				BufferSizeBytes = (BufferSizeBytes > MaxBufferToProcess) ? MaxBufferToProcess : BufferSizeBytes;
-				Buffer.SetNumUninitialized(IntCastChecked<int32>(BufferSizeBytes));
+				LogReader->Seek(ReadStart);
+				Buffer.SetNumUninitialized(BytesToRead);
 				LogReader->Serialize(Buffer.GetData(), Buffer.Num());
 
-				FString LogString = FString(StringCast<TCHAR>(reinterpret_cast<const char*>(Buffer.GetData()), Buffer.Num()));
+				FString LogString = FString(Buffer.Num(), StringCast<TCHAR>(reinterpret_cast<const char*>(Buffer.GetData()), Buffer.Num()).Get());
 				FRegexMatcher Regex(LogProgressPattern, LogString);
 
 				while (Regex.FindNext())
@@ -671,7 +670,7 @@ void UWorldPartitionMeshPartitionBuilder::ExecuteModifierGroupTasks(TIndirectArr
 				case ETaskState::Ready:
 					if (NumTasksInFlight < MaxTasksInFlight && !FWorldPartitionHelpers::HasExceededMaxMemory())
 					{
-						Task.ReuseCompiledSectionActorDescGuid = FGuid();
+						Task.ReuseCompiledSectionActorDescGuids.Empty();
 
 						// TODO [chris.tchou] : this DOES NOT include GUIDs from modifiers inside level instances...	so it doesn't detect adding a new instance inside a level instance.
 						Task.BuildInfo.ModifierSetHash = Task.ModifierGroup.ComputeModifierSetHash();
@@ -679,15 +678,11 @@ void UWorldPartitionMeshPartitionBuilder::ExecuteModifierGroupTasks(TIndirectArr
 						// try to reuse by package hashes (which doesn't require loading any modifiers, so we can early out if this works)
 						if ((ReuseMode == EMeshPartitionReuseMode::ByPackageHash) || (ReuseMode == EMeshPartitionReuseMode::ValidateHashes))
 						{
-							if (const FCompiledSectionDescriptor* MatchingCompiledSection = FindMatchingCompiledSectionFromPackageHashes(Task.OldCompiledSectionDescriptors, Task.BuildInfo, AssetRegistry))
-							{
-								// early out, reuse existing
-								Task.ReuseCompiledSectionActorDescGuid = MatchingCompiledSection->ActorDescGuid;
-							}
+							Task.ReuseCompiledSectionActorDescGuids = FindMatchingCompiledSectionsFromPackageHashes(Task.OldCompiledSectionDescriptors, Task.BuildInfo, AssetRegistry);
 						}
 
 						// if we're not early out (or if we're in validate mode)
-						if (!Task.ReuseCompiledSectionActorDescGuid.IsValid() || (ReuseMode == EMeshPartitionReuseMode::ValidateHashes))
+						if (Task.ReuseCompiledSectionActorDescGuids.IsEmpty() || (ReuseMode == EMeshPartitionReuseMode::ValidateHashes))
 						{
 							// load the modifiers, including all modifiers within level instances
 							LoadActors(InMeshPartition, Task);
@@ -698,34 +693,61 @@ void UWorldPartitionMeshPartitionBuilder::ExecuteModifierGroupTasks(TIndirectArr
 
 							if (ReuseMode == EMeshPartitionReuseMode::ByModifierHash)
 							{
-								Task.ReuseCompiledSectionActorDescGuid = FindMatchingCompiledSectionFromModifierHash(Task.OldCompiledSectionDescriptors, Task.BuildInfo);
+								Task.ReuseCompiledSectionActorDescGuids = FindMatchingCompiledSectionsFromModifierHash(Task.OldCompiledSectionDescriptors, Task.BuildInfo);
 							}
 							else if (ReuseMode == EMeshPartitionReuseMode::ValidateHashes)
 							{
-								FGuid ModifierHashReuseGuid = FindMatchingCompiledSectionFromModifierHash(Task.OldCompiledSectionDescriptors, Task.BuildInfo);
+								TArray<FGuid> ModifierHashReuseGuids = FindMatchingCompiledSectionsFromModifierHash(Task.OldCompiledSectionDescriptors, Task.BuildInfo);
 
 								// ModifierHash may find a match when PackageHash does not, that is expected.
 								// But we can check:
 
-								// if PackageHash finds a match, then modifier hash should agree
-								if (Task.ReuseCompiledSectionActorDescGuid.IsValid())
+								// Sort both for set comparison (order may differ between search paths)
+								Task.ReuseCompiledSectionActorDescGuids.Sort();
+								ModifierHashReuseGuids.Sort();
+
+								// if PackageHash finds matches, then modifier hash should agree
+								if (!Task.ReuseCompiledSectionActorDescGuids.IsEmpty())
 								{
 									// (if not we have a bug in our hash/dependency calculations)
-									check(Task.ReuseCompiledSectionActorDescGuid == ModifierHashReuseGuid);
+									check(Task.ReuseCompiledSectionActorDescGuids == ModifierHashReuseGuids);
 								}
 
 								// if ModifierHash finds NO match, then PackageHash should agree
-								if (!ModifierHashReuseGuid.IsValid())
+								if (ModifierHashReuseGuids.IsEmpty())
 								{
 									// (if not we have a bug in our hash/dependency calculations)
-									check(!Task.ReuseCompiledSectionActorDescGuid.IsValid());
+									check(Task.ReuseCompiledSectionActorDescGuids.Num() == 0);
 								}
 
 								// TODO: if we record PackageHash and/or ModifierHash per Package/Modifier, then we can provide better info on which is causing the issue here
 							}
 						}
 
-						if (Task.ReuseCompiledSectionActorDescGuid.IsValid())
+						// For grid-split groups, reject partial reuse — all cells must match or none do.
+						if (!Task.ReuseCompiledSectionActorDescGuids.IsEmpty() && (Task.ResolvedGridCellSize > 0))
+						{
+							int32 TotalStructuralMatches = 0;
+							
+							for (const FCompiledSectionDescriptor& Desc : Task.OldCompiledSectionDescriptors)
+							{
+								if (Desc.Info.TargetsSameCompiledSectionAs(Task.BuildInfo) &&
+									Desc.Info.BuildVariantHash == Task.BuildInfo.BuildVariantHash &&
+									Desc.Info.ModifierSetHash == Task.BuildInfo.ModifierSetHash)
+								{
+									TotalStructuralMatches++;
+								}
+							}
+							
+							if (Task.ReuseCompiledSectionActorDescGuids.Num() < TotalStructuralMatches)
+							{
+								UE_LOG(LogMegaMeshEditor, Warning, TEXT("Grid-split reuse: only %d of %d cells have up-to-date hashes, forcing rebuild"),
+									Task.ReuseCompiledSectionActorDescGuids.Num(), TotalStructuralMatches);
+								Task.ReuseCompiledSectionActorDescGuids.Empty();
+							}
+						}
+
+						if (!Task.ReuseCompiledSectionActorDescGuids.IsEmpty())
 						{
 							Task.Result = ETaskResult::ReuseExisting;
 							Task.State = ETaskState::Completed;
@@ -789,23 +811,39 @@ void UWorldPartitionMeshPartitionBuilder::ExecuteModifierGroupTasks(TIndirectArr
 					break;
 								
 				case ETaskState::ModifiersProcessed:
-					Task.CompiledSection = Task.MeshPartitionEditorComponent->PrepareMegaMeshCompiledSection(Task.BuildInfo, Task.BuildVariant);
-					Task.MeshPartitionEditorComponent->BuildMegaMeshCompiledSectionTextures(Task.CompiledSection, Task.ModifierGroup, Task.ProcessModifierTask.GetResult());
-
 					{
+						MeshPartition::FMeshData& FullMesh = Task.ProcessModifierTask.GetResult();
+						FPrepareCompiledSectionsParams CompiledSectionsParams {.FullMesh = FullMesh, .GridCellSize = static_cast<double>(Task.ResolvedGridCellSize)};
+
+						Task.CompiledSections = Task.MeshPartitionEditorComponent->PrepareCompiledSections(Task.BuildInfo, Task.BuildVariant, CompiledSectionsParams);
+
+						if (Task.CompiledSections.IsEmpty())
+						{
+							UE_LOG(LogMegaMeshEditor, Error, TEXT("PrepareCompiledSections produced no sections for modifier group. Skipping."));
+							Task.Result = ETaskResult::BuildFail;
+							Task.State = ETaskState::Completed;
+							NumTasksInFlight--;
+							break;
+						}
+
 						TArray<TWeakObjectPtr<MeshPartition::UModifierComponent>> ModifierPtrs = Task.ModifierGroup.AllResolvedModifierPtrs();
-						Task.MeshPartitionEditorComponent->PostBuildSectionMesh(Task.CompiledSection, Task.ProcessModifierTask.GetResult(), ModifierPtrs);
+						TArray<MeshPartition::FTransformerUnit> InitialUnits;
+						InitialUnits.Reserve(Task.CompiledSections.Num());
 
-						MeshPartition::FTransformerUnit InitialUnit;
+						for (int32 CompiledSectionIndex = 0; CompiledSectionIndex < Task.CompiledSections.Num(); ++CompiledSectionIndex)
+						{
+							MeshPartition::ACompiledSection* Section = Task.CompiledSections[CompiledSectionIndex];
+							MeshPartition::FMeshData* CellMeshPtr = CompiledSectionsParams.OutCellMeshes.Find(Section->GetBuildInfo().GridCellCoord);
+							MeshPartition::FMeshData& SectionMesh = CellMeshPtr ? *CellMeshPtr : FullMesh;
 
-						InitialUnit.Section = Task.CompiledSection;
-						InitialUnit.MeshData = MakeShared<const MeshPartition::FMeshData>(MoveTemp(Task.ProcessModifierTask.GetResult()));
-						InitialUnit.bShouldRecomputeNormals = true;
-						InitialUnit.bShouldRecomputeTangents = true;
+							Task.MeshPartitionEditorComponent->BuildMegaMeshCompiledSectionTextures(Section, Task.ModifierGroup, SectionMesh);
+							Task.MeshPartitionEditorComponent->PostBuildSectionMesh(Section, SectionMesh, ModifierPtrs);
 
-						MeshPartition::UMeshPartitionDefinition* Definition = Task.MeshPartitionEditorComponent->GetMegaMeshDefinition();
+							InitialUnits.Add(MeshPartition::MakeTransformerUnit(Section, MakeShared<const MeshPartition::FMeshData>(MoveTemp(SectionMesh))));
+						}
 
-						Task.TransformerContext = Task.MeshPartitionEditorComponent->LaunchTransformers(MoveTemp(InitialUnit), Definition, Task.BuildVariant);
+						MeshPartition::UMeshPartitionDefinition* Definition = Task.MeshPartitionEditorComponent->GetMegaMeshDefinition();
+						Task.TransformerContext = Task.MeshPartitionEditorComponent->LaunchTransformers(MoveTemp(InitialUnits), Definition, Task.BuildVariant);
 					}
 
 					if (Task.TransformerContext.IsValid())
@@ -828,14 +866,16 @@ void UWorldPartitionMeshPartitionBuilder::ExecuteModifierGroupTasks(TIndirectArr
 					break;
 
 				case ETaskState::TransformersExecuted:
-					Task.ModifiedPackages.Emplace(Task.CompiledSection->GetPackage());
-
 					{
 						TArray<TWeakObjectPtr<MeshPartition::UModifierComponent>> ModifierPtrs = Task.ModifierGroup.AllResolvedModifierPtrs();
-						Task.MeshPartitionEditorComponent->PostProcessSection(Task.CompiledSection, ModifierPtrs);
-					}
 
-					AssignSectionToPlatformDataLayers(Task.CompiledSection);
+						for (MeshPartition::ACompiledSection* CellSection : Task.CompiledSections)
+						{
+							Task.ModifiedPackages.Emplace(CellSection->GetPackage());
+							Task.MeshPartitionEditorComponent->PostProcessSection(CellSection, ModifierPtrs);
+							AssignSectionToPlatformDataLayers(CellSection);
+						}
+					}
 
 					// save and reclaim memory
 					SaveAndCollectGarbage(InWorld, InPackageHelper, Task);
@@ -980,15 +1020,19 @@ void UWorldPartitionMeshPartitionBuilder::SaveAndCollectGarbage(UWorld* InWorld,
 	TRACE_CPUPROFILER_EVENT_SCOPE(UWorldPartitionMeshPartitionBuilder::SaveAndCollectGarbage);
 
 	TArray<UObject*> ObjectsToFinish;
-	ForEachObjectWithOuter(InModifierGroupTask.CompiledSection, [&ObjectsToFinish](UObject* Object)
+
+	for (MeshPartition::ACompiledSection* Section : InModifierGroupTask.CompiledSections)
 	{
-		if (const IInterface_AsyncCompilation* AsyncCompilationIF = Cast<IInterface_AsyncCompilation>(Object))
+		ForEachObjectWithOuter(Section, [&ObjectsToFinish](UObject* Object)
 		{
-			ObjectsToFinish.Emplace(Object);
-		}
-	
-		return true;
-	});
+			if (const IInterface_AsyncCompilation* AsyncCompilationIF = Cast<IInterface_AsyncCompilation>(Object))
+			{
+				ObjectsToFinish.Emplace(Object);
+			}
+
+			return true;
+		});
+	}
 
 	if (!ObjectsToFinish.IsEmpty())
 	{
@@ -1000,11 +1044,14 @@ void UWorldPartitionMeshPartitionBuilder::SaveAndCollectGarbage(UWorld* InWorld,
 		SavePackages(InModifierGroupTask.ModifiedPackages, InPackageHelper);
 	}
 
-	// Destroy the newly created compiled section _after_ the save so that it can be fully unloaded by the next GC.
-	InWorld->DestroyActor(InModifierGroupTask.CompiledSection);
+	// Destroy all newly created compiled sections _after_ the save so they can be fully unloaded by the next GC.
+	for (MeshPartition::ACompiledSection* Section : InModifierGroupTask.CompiledSections)
+	{
+		InWorld->DestroyActor(Section);
+	}
 
 	// clear out pointers, since the objects may be going away in the garbage collection
-	InModifierGroupTask.CompiledSection = nullptr;
+	InModifierGroupTask.CompiledSections.Reset();
 	InModifierGroupTask.ModifiedPackages.Empty();
 	InModifierGroupTask.ModifierOwners.Empty();
 	InModifierGroupTask.ActorRefs.Empty();
diff --git a/Engine/Plugins/Experimental/MeshPartition/Source/MeshPartitionEditor/Public/MeshPartitionChannelCollection.h b/Engine/Plugins/Experimental/MeshPartition/Source/MeshPartitionEditor/Public/MeshPartitionChannelCollection.h
--- a/Engine/Plugins/Experimental/MeshPartition/Source/MeshPartitionEditor/Public/MeshPartitionChannelCollection.h
+++ b/Engine/Plugins/Experimental/MeshPartition/Source/MeshPartitionEditor/Public/MeshPartitionChannelCollection.h
@@ -16,6 +16,7 @@ enum EPixelFormat : uint8;
 
 namespace UE::MeshPartition
 {
+class ACompiledSection;
 class UMeshPartitionDefinition;
 class UMeshPartitionEditorComponent;
 class UModifierComponent;
@@ -66,6 +67,9 @@ struct FSectionChannels
 	FSectionChannelDomainData Domain;
 };
 
+/** Apply pre-built channel data (texture, texcoord metrics, table) to a compiled section. */
+UE_API void ApplyChannels(ACompiledSection* InSection, const FSectionChannels& InChannels, const float InMaterialCacheTexelSize);
+
 /**
 * FMegaMeshChannelCollection contains functions and objects needed to build the rendering elements 
 * representing the masks and channels for the sections of the mega mesh
diff --git a/Engine/Plugins/Experimental/MeshPartition/Source/MeshPartitionEditor/Public/MeshPartitionDependencyContext.h b/Engine/Plugins/Experimental/MeshPartition/Source/MeshPartitionEditor/Public/MeshPartitionDependencyContext.h
--- a/Engine/Plugins/Experimental/MeshPartition/Source/MeshPartitionEditor/Public/MeshPartitionDependencyContext.h
+++ b/Engine/Plugins/Experimental/MeshPartition/Source/MeshPartitionEditor/Public/MeshPartitionDependencyContext.h
@@ -83,6 +83,8 @@ public:
 class FDependencyHash : public IDependencyInterface
 {
 public:
+	using IDependencyInterface::operator+=;
+
 	FDependencyHash() = default;
 	virtual ~FDependencyHash() override = default;
 
diff --git a/Engine/Plugins/Experimental/MeshPartition/Source/MeshPartitionEditor/Public/MeshPartitionEditorComponent.h b/Engine/Plugins/Experimental/MeshPartition/Source/MeshPartitionEditor/Public/MeshPartitionEditorComponent.h
--- a/Engine/Plugins/Experimental/MeshPartition/Source/MeshPartitionEditor/Public/MeshPartitionEditorComponent.h
+++ b/Engine/Plugins/Experimental/MeshPartition/Source/MeshPartitionEditor/Public/MeshPartitionEditorComponent.h
@@ -36,6 +36,7 @@ class AInteractiveSection;
 class APreviewSection;
 class UMeshPartitionDefinition;
 class UModifierComponent;
+struct FSectionChannels;
 
 extern MESHPARTITIONEDITOR_API TAutoConsoleVariable<bool> CVarMegaMeshEnablePreviewSimplification;
 extern MESHPARTITIONEDITOR_API TAutoConsoleVariable<float> CVarMegaMeshPreviewSimplificationEdgeLength;
@@ -56,6 +57,26 @@ struct FPreviewSectionBuildContext
 	TArray<MeshPartition::FBuildTaskHandle> BuildTasks;
 };
 
+/** Parameters for PrepareCompiledSections. */
+struct FPrepareCompiledSectionsParams
+{
+	// Full mesh for this modifier group.
+	const MeshPartition::FMeshData& FullMesh;
+
+	// Resolved WP grid cell size (cm). 0 = no grid splitting.
+	// Set by the builder from ResolveGridCellSizeFromPipeline().
+	double GridCellSize = 0.0;
+
+	// Receives per-cell sub-meshes after the call, keyed by absolute grid coordinate.
+	TMap<FIntVector, MeshPartition::FMeshData> OutCellMeshes;
+
+	// Existing section actor to reuse instead of spawning a new one (single-section only).
+	MeshPartition::ACompiledSection* ReuseSection = nullptr;
+
+	// Mobility for spawned actors (applies to single-section path only; grid always uses static).
+	bool bUseStaticMobility = true;
+};
+
 UCLASS(MinimalAPI, ClassGroup=(Custom), meta=(BlueprintSpawnableComponent))
 class UMeshPartitionEditorComponent : public UMeshPartitionComponent
 {
@@ -102,10 +123,19 @@ public:
 	UE_API void BuildMegaMeshPreviewSections(TConstArrayView<const FBox> InBounds);
 
 	/**
-	 * Spawns a MeshPartition::ACompiledSection
-	 * @return The spawned MeshPartition::ACompiledSection
+	 * Spawns one or more ACompiledSection actors for a modifier group.
+	 * Dispatches to a grid-aligned split or a single section based on InBuildVariant.bSplitSectionsToMatchWorldPartitionRuntimeGrid and InParams.GridCellSize:
+	 *   - Grid path (bSplitSectionsToMatchWorldPartitionRuntimeGrid=true, InParams.GridCellSize>0):
+	 *       Splits FullMesh by GridCellSize, spawns one section per intersecting cell,
+	 *       populates OutCellMeshes with the per-cell sub-meshes.
+	 *   - Single path (all other cases):
+	 *       Spawns one section. Warns if bSplitSectionsToMatchWorldPartitionRuntimeGrid=true but GridCellSize is 0.
+	 * @return Array of spawned/reused sections; never empty on success.
 	 */
-	UE_API MeshPartition::ACompiledSection* PrepareMegaMeshCompiledSection(const MeshPartition::FCompiledSectionBuildInfo& InBuildInfo, const MeshPartition::FCompiledSectionBuildVariant& InBuildVariant, MeshPartition::ACompiledSection* InReuseCompiledSection = nullptr, bool bInUseStaticMobility = true) const;
+	UE_API TArray<MeshPartition::ACompiledSection*> PrepareCompiledSections(
+		const MeshPartition::FCompiledSectionBuildInfo& InBuildInfo,
+		const MeshPartition::FCompiledSectionBuildVariant& InBuildVariant,
+		MeshPartition::FPrepareCompiledSectionsParams& InParams) const;
 
 	/**
      * Only used when building a preview section. Executes transformers on a preview mesh component.
@@ -118,13 +148,13 @@ public:
 	UE_API void ExecuteTransformers(MeshPartition::UPreviewMeshComponent* InPreviewMeshComponent, const MeshPartition::UMeshPartitionDefinition* InDefinition, const MeshPartition::FCommonBuildVariant& InBuildVariant);
 
 	/**
-	 * Launches transformers associated with this build variant.
-	 * @param InTransformerUnit The initial transformer unit for the pipeline. The pipeline will take ownership of this unit.
+	 * Launches transformers associated with this build variant, seeding the pipeline with the provided initial units.
+	 * @param InTransformerUnits The initial transformer units for the pipeline. The pipeline will take ownership of all units.
 	 * @param InDefinition The definition used to launch the transformer pipeline.
 	 * @param InBuildVariant The variant used to launch the transformer pipeline.
 	 * @return If successful, launches the transformer pipeline and returns a new FTransformerContext.
 	 */
-	UE_API TUniquePtr<MeshPartition::FTransformerContext> LaunchTransformers(MeshPartition::FTransformerUnit&& InTransformerUnit, const MeshPartition::UMeshPartitionDefinition* InDefinition, const MeshPartition::FCommonBuildVariant& InBuildVariant);
+	UE_API TUniquePtr<MeshPartition::FTransformerContext> LaunchTransformers(TArray<MeshPartition::FTransformerUnit>&& InTransformerUnits, const MeshPartition::UMeshPartitionDefinition* InDefinition, const MeshPartition::FCommonBuildVariant& InBuildVariant);
 
 	/**
 	* Builds the collection of textures containing the channel map information and binds it to the compiled section's material.
@@ -160,7 +190,22 @@ public:
 	 * @param InModifiers The modifiers applied to this section.
 	 */
 	UE_API void PostProcessSection(AActor* InSection, TConstArrayView<TWeakObjectPtr<MeshPartition::UModifierComponent>> InModifiers) const;
-	
+
+	/**
+	 * Runs the full synchronous post-build pipeline on a compiled section:
+	 * PrepareCompiledSections → ApplyChannels → PostBuildSectionMesh → Transformers (wait) → PostProcessSection.
+	 * Resolves Definition and BuildVariant from the section's BuildInfo internally.
+	 * Callers must ensure the section's BuildInfo is up-to-date before calling.
+	 * @param InSection The section to finalize.
+	 * @param InMesh The built mesh data (shared ownership for transformer pipeline).
+	 * @param InChannels Pre-built channel data (texture, texcoord metrics, table).
+	 * @param InModifierGroup The modifier group for this section.
+	 */
+	UE_API void ApplyBuiltMeshToSection(MeshPartition::ACompiledSection*				InSection,
+										TSharedPtr<const MeshPartition::FMeshData>		InMesh,
+										const MeshPartition::FSectionChannels&			InChannels,
+										const MeshPartition::FModifierGroup&			InModifierGroup);
+
 	/**
 	* Called when a modifier contributing to the MegaMesh changed. Will invalidate and trigger a rebuilt of changed areas.
 	* @param InChangedModifier The modifier which has changed.
diff --git a/Engine/Plugins/Experimental/MeshPartition/Source/MeshPartitionEditor/Public/MeshPartitionEditorUtils.h b/Engine/Plugins/Experimental/MeshPartition/Source/MeshPartitionEditor/Public/MeshPartitionEditorUtils.h
--- a/Engine/Plugins/Experimental/MeshPartition/Source/MeshPartitionEditor/Public/MeshPartitionEditorUtils.h
+++ b/Engine/Plugins/Experimental/MeshPartition/Source/MeshPartitionEditor/Public/MeshPartitionEditorUtils.h
@@ -53,25 +53,20 @@ namespace UE::MeshPartition::EditorUtils
 	bool ClassHashIsUpToDate(const MeshPartition::FCompiledSectionBuildInfo& BuildInfo);
 
 	/**
-	 * Determine if there is a pre-existing compiled section that has the same PackageHash (and ClassHash) and can be re-used as the build result.
-	 * This does not require loading the modifiers, as we can check the asset registry to determine if the packages have changed, which gives us much faster early out.
-	 * This requires all modifiers to implement GatherDependencies() to produce a valid Package hash.
-	 *
-	 * @param CompiledSectionDescriptorsToSearch the set of compiled section descriptors to search for a match
-	 * @param BuildInfo the build info we are looking to find a compiled section for
-	 * @return the matching compiled section descriptor, if found, or null otherwise
+	 * Find all pre-existing compiled sections that match BuildInfo by identity, BuildVariantHash, ModifierSetHash,
+	 * PackageHash, and ClassHash. Returns GUIDs for all matching siblings (e.g. N grid-split sections from the same group).
+	 * Empty result means no reuse is possible. Does not require loading modifiers.
 	 */
-	const FCompiledSectionDescriptor* FindMatchingCompiledSectionFromPackageHashes(
+	TArray<FGuid> FindMatchingCompiledSectionsFromPackageHashes(
 		TArrayView<const FCompiledSectionDescriptor> CompiledSectionDescriptorsToSearch,
 		const MeshPartition::FCompiledSectionBuildInfo& BuildInfo,
 		IAssetRegistry& AssetRegistry);
 
 	/**
-	 * Determine if there is a pre-existing compiled section that has the same ModifierHash and can be re-used as the build result.
-	 * @param Task the modifier group we're building
-	 * @return a valid ActorDescGuid for a compiled section we can re-use, or an invalid FGuid if no reuse is possible
+	 * Find all pre-existing compiled sections that match BuildInfo by identity, BuildVariantHash, and ModifiersHash.
+	 * Returns GUIDs for all matching siblings. Empty result means no reuse is possible.
 	 */
-	FGuid FindMatchingCompiledSectionFromModifierHash(
+	TArray<FGuid> FindMatchingCompiledSectionsFromModifierHash(
 		TArrayView<const FCompiledSectionDescriptor> CompiledSectionDescriptorsToSearch,
 		const MeshPartition::FCompiledSectionBuildInfo& BuildInfo);
 
diff --git a/Engine/Plugins/Experimental/MeshPartition/Source/MeshPartitionEditor/Public/MeshPartitionMeshBuilder.h b/Engine/Plugins/Experimental/MeshPartition/Source/MeshPartitionEditor/Public/MeshPartitionMeshBuilder.h
--- a/Engine/Plugins/Experimental/MeshPartition/Source/MeshPartitionEditor/Public/MeshPartitionMeshBuilder.h
+++ b/Engine/Plugins/Experimental/MeshPartition/Source/MeshPartitionEditor/Public/MeshPartitionMeshBuilder.h
@@ -56,6 +56,38 @@ MESHPARTITIONEDITOR_API void Wait(TConstArrayView<MeshPartition::FBuildTaskHandl
 MESHPARTITIONEDITOR_API bool AreAllTasksComplete(TConstArrayView<MeshPartition::FBuildTaskHandle> InTasks);
 }
 
+namespace UE::MeshPartition::GridHelpers
+{
+/** Result of grid dimension computation: snapped origin + per-axis counts. */
+struct FGridDimensions
+{
+	FVector SnappedMin;
+	FIntVector OriginCoord;		// Absolute grid coordinate of the cell at local index (0,0,0): floor(SnappedMin / CellSize)
+	FIntVector CellNumber;
+	int32 TotalCells;
+};
+
+/**
+ * Computes grid dimensions from bounds and cell size.
+ * Snaps bounds min to grid-aligned coordinate (floor), then divides the snapped extents.
+ * This is the single source of truth for grid dimension logic — used by both BuildSections and EstimateGridCells.
+ * @param InBounds The bounding box to subdivide.
+ * @param InCellSize The grid cell size in unreal units. Must be > 0.
+ * @return Grid dimensions including snapped origin and per-axis cell counts.
+ */
+MESHPARTITIONEDITOR_API FGridDimensions ComputeGridDimensions(const FBox& InBounds, double InCellSize);
+
+/**
+ * Splits a mesh into grid-aligned cells and returns a map keyed by absolute world-space grid coordinates.
+ * Each coordinate component = floor(cellWorldMin / InCellSize), making it independent of the input bounds.
+ * Empty cells are excluded from the result.
+ * @param InMesh The mesh to split.
+ * @param InCellSize The grid cell size in unreal units. Must be > 0.
+ * @return Map from absolute grid coordinate to cell mesh. Only non-empty cells are included.
+ */
+MESHPARTITIONEDITOR_API TMap<FIntVector, MeshPartition::FMeshData> BuildGridCellMeshes(const MeshPartition::FMeshData& InMesh, double InCellSize);
+}
+
 namespace UE::MeshPartition::BuildHelpers
 {
 /**
@@ -89,26 +121,15 @@ UE::Tasks::FTask ProcessModifierGroup(
 TArray<MeshPartition::FModifierDesc> InitializeModifierDescriptors(const MeshPartition::FBuilderSettings& InSettings);
 
 
-/**
- * Computes how many sections of size InSectionSize can fit into the given bounds.
- * At this moment is method is treating everything as a 2D grid. Will change for 3D.
- * @param InBounds The bounds to divide.
- * @param InSectionSize The section size in unreal units.
- * @param OutNumberOfSectionsX The number of sections for the X axis.
- * @param OutNumberOfSectionsY The number of sections for the Y axis.
- * @param OutNumberOfSectionsZ The number of sections for the Z axis.
- * @return The total number of sections (OutNumberOfSectionsX * OutNumberOfSectionsY * OutNumberOfSectionsZ).
- */
-int32 GetTotalNumberOfSections(const FBox& InBounds, const double InSectionSize, int32& OutNumberOfSectionsX, int32& OutNumberOfSectionsY, int32& OutNumberOfSectionsZ);
-
 /**
  * Splits the given mesh into multiple meshes of maximum InSectionSize * InSectionSize.
  * At this moment is method is treating everything as a 2D grid. Will change for 3D.
  * @param InMesh The mesh to split.
  * @param InSectionSize The section size in unreal units.
+ * @param bInFilterEmptyMeshes When true, empty cells are removed from the output; when false all cells are included.
  * @return An array containing the splitted mesh.
  */
-TArray<MeshPartition::FMeshData> BuildSections(const MeshPartition::FMeshData& InMesh, const double InSectionSize);
+TArray<MeshPartition::FMeshData> BuildSections(const MeshPartition::FMeshData& InMesh, const double InSectionSize, const bool bInFilterEmptyMeshes = true);
 
 /**
  * Produces a simplified version of the input mesh by collapsing edges whose lengths fall below the requested tolerance.
diff --git a/Engine/Plugins/Experimental/MeshPartition/Source/MeshPartitionEditor/Public/MeshPartitionSeparateWorldBuilder.h b/Engine/Plugins/Experimental/MeshPartition/Source/MeshPartitionEditor/Public/MeshPartitionSeparateWorldBuilder.h
--- a/Engine/Plugins/Experimental/MeshPartition/Source/MeshPartitionEditor/Public/MeshPartitionSeparateWorldBuilder.h
+++ b/Engine/Plugins/Experimental/MeshPartition/Source/MeshPartitionEditor/Public/MeshPartitionSeparateWorldBuilder.h
@@ -9,6 +9,7 @@
 namespace UE::MeshPartition
 {
 class ACompiledSection;
+struct FModifierGroupBuildCache;
 
 USTRUCT()
 struct FSeparateWorldBuilder
@@ -18,14 +19,14 @@ struct FSeparateWorldBuilder
 	~FSeparateWorldBuilder();
 
 	// Add a compiled section build request to the queue
-	void RequestCompiledSectionBuild(MeshPartition::ACompiledSection* CompiledSection);
+	void RequestCompiledSectionBuild(MeshPartition::ACompiledSection* InCompiledSection);
 
 	// Execute queued compiled section builds
 	void Tick();
 
 private:
 	void InitializeToBuildOutOfWorldPackage(const FString& InWorldPackageName);
-	void BuildCompiledSectionMesh(MeshPartition::ACompiledSection* CompiledSection, FName BuildVariantName);
+	void BuildCompiledSectionMesh(MeshPartition::ACompiledSection* InCompiledSection, FName InBuildVariantName);
 
 	// Current package used as a container for the separate world
 	UPROPERTY()
@@ -41,5 +42,9 @@ private:
 	// other state that doesn't hold GC refs
 	TArray<MeshPartition::FBuildTaskHandle> TaskHandles;
 	TArray<FGuid> AllMegaMeshGuidsInWorld;
+
+	// Per-group build cache for grid-split sections (avoids N full mesh rebuilds per group).
+	// Created lazily on first grid-split build, cleared after each Tick().
+	TSharedPtr<FModifierGroupBuildCache> GroupBuildCache;
 };
 } // namespace UE::MeshPartition
\ No newline at end of file
diff --git a/Engine/Plugins/Experimental/MeshPartition/Source/MeshPartitionEditor/Public/MeshPartitionWPActorPropertiesTransformer.h b/Engine/Plugins/Experimental/MeshPartition/Source/MeshPartitionEditor/Public/MeshPartitionWPActorPropertiesTransformer.h
--- a/Engine/Plugins/Experimental/MeshPartition/Source/MeshPartitionEditor/Public/MeshPartitionWPActorPropertiesTransformer.h
+++ b/Engine/Plugins/Experimental/MeshPartition/Source/MeshPartitionEditor/Public/MeshPartitionWPActorPropertiesTransformer.h
@@ -30,6 +30,8 @@ public:
 
 	UE_API virtual void GatherDependencies(MeshPartition::IDependencyInterface& InDependencies) const override;
 
+	FName GetRuntimeGrid() const { return RuntimeGrid; }
+
 private:
 	void SetActorProperties(const MeshPartition::FTransformerContext& InTransformerContext) const;
 
diff --git a/Engine/Plugins/Experimental/MeshPartition/Source/MeshPartitionEditor/Public/MeshPartitionWorldPartitionHelpers.h b/Engine/Plugins/Experimental/MeshPartition/Source/MeshPartitionEditor/Public/MeshPartitionWorldPartitionHelpers.h
--- a/Engine/Plugins/Experimental/MeshPartition/Source/MeshPartitionEditor/Public/MeshPartitionWorldPartitionHelpers.h
+++ b/Engine/Plugins/Experimental/MeshPartition/Source/MeshPartitionEditor/Public/MeshPartitionWorldPartitionHelpers.h
@@ -5,6 +5,7 @@
 #include "CoreMinimal.h"
 #include "WorldPartition/WorldPartitionHandle.h"
 
+class UWorld;
 class UWorldPartition;
 class ULevelStreaming;
 
@@ -12,11 +13,35 @@ namespace UE::MeshPartition
 {
 	class AMeshPartition;
 	class UModifierComponent;
+	struct FCompiledSectionBuildVariant;
 	struct FModifierGroup;
 }
 
 namespace UE::MeshPartition::WorldPartitionHelpers
 {
+	/** Returns the cell size of a runtime partition grid, or 0 if the grid doesn't exist or isn't an LHGrid.
+	 *  Passing NAME_None resolves to the main (first) grid. */
+	uint32 GetGridCellSize(const UWorld* InWorld, FName InGridName);
+
+	/**
+	 * Resolves the WP grid cell size for a build variant by scanning its transformer pipeline.
+	 * Finds the last FWPActorPropertiesTransformer, reads its RuntimeGrid name, and resolves the cell size.
+	 * Fallback: if the named grid is not found, tries the main grid (NAME_None).
+	 * Returns 0 if no WP transformer is found or no grid resolves.
+	 */
+	uint32 ResolveGridCellSizeFromPipeline(const MeshPartition::FCompiledSectionBuildVariant& InBuildVariant, const UWorld* InWorld);
+
+	/** Estimated grid cell with its absolute world-space coordinate and bounding box. */
+	struct FGridCellEstimate
+	{
+		FIntVector GridCellCoord;
+		FBox CellBounds;
+	};
+
+	/** Computes per-cell bounds from a bounding box and cell size, using the same snapping as BuildSections.
+	 *  Returns all cells that the bounds overlap (some may contain no actual geometry). */
+	TArray<FGridCellEstimate> EstimateGridCells(const FBox& InGroupBounds, double InCellSize);
+
 	TSet<ULevelStreaming*> LoadAllLevelInstances(UWorldPartition* WorldPartition, TSet<FWorldPartitionReference>& InOutActorRefs);
 
 	void LoadAllActorsFromStreamingLevels(
diff --git a/Engine/Plugins/Experimental/MeshPartition/Source/MeshPartitionEditor/Public/MeshPartitionWorldUpdater.h b/Engine/Plugins/Experimental/MeshPartition/Source/MeshPartitionEditor/Public/MeshPartitionWorldUpdater.h
--- a/Engine/Plugins/Experimental/MeshPartition/Source/MeshPartitionEditor/Public/MeshPartitionWorldUpdater.h
+++ b/Engine/Plugins/Experimental/MeshPartition/Source/MeshPartitionEditor/Public/MeshPartitionWorldUpdater.h
@@ -18,6 +18,7 @@ namespace UE::MeshPartition
 		bool bIsPlaceholder : 1 = false;		// is a PIE placeholder
 		bool bNonTargetVariant : 1 = false;		// corresponds to a build variant that is not a target of this updater
 		bool bMismatchedModifiers : 1 = false;	// modifier set hash does not match any target group (set hash is derived from the modifier paths)
+		bool bBuildVariantHashMatches : 1 = false;	// build variant hash matches (definition/variant settings unchanged)
 		bool bPackageHashMatches : 1 = false;	// package hash matches
 		bool bClassHashMatches : 1 = false;  	// class hash matches
 		bool bCannotFindFile : 1 = false;		// cannot find the compiled section file (asset registry does not know about it)
@@ -34,7 +35,7 @@ namespace UE::MeshPartition
 		friend struct FMeshPartitionUpdater;
 
 	public:
-		FBuildVariantUpdater(const FGuid& InBuildKey, const UMeshPartitionDefinition& InMeshPartitionDefinition, const MeshPartition::FModifierGrouping& InTargetModifierGrouping, TSharedPtr<MeshPartition::FCachedDescriptors> InAllMeshPartitionDescriptors, IAssetRegistry& InAssetRegistry);
+		FBuildVariantUpdater(const FGuid& InBuildKey, const UMeshPartitionDefinition& InMeshPartitionDefinition, const MeshPartition::FModifierGrouping& InTargetModifierGrouping, TSharedPtr<MeshPartition::FCachedDescriptors> InAllMeshPartitionDescriptors, IAssetRegistry& InAssetRegistry, const UWorld* InWorld);
 
 	private:
 		// partial BuildInfo for this build variant (used to initialize PIE placeholders, and also stores info like BuildVariantName)
@@ -49,12 +50,16 @@ namespace UE::MeshPartition
 		// mapping from modifier set hash to target group (index into TargetGrouping)
 		TMap<FGuid, int32> ModifierSetHashToTargetGroupIndex;
 
-		// This array is 1:1 with the array of groups in TargetGrouping, and provides the best reusable compiled section for each group (or -1 if there is none)
-		TArray<int32> BestReusableCompiledSectionIndexForGroup;
+		// Per-group reuse tracking: maps GridCellCoord → CompiledSectionIndex in the MeshPartitionUpdater's CompiledSections array.
+		// For non-grid sections, the key is InvalidGridCellCoord. Empty map = no reusable sections for that group.
+		TArray<TMap<FIntVector, int32>> ReusableCompiledSectionIndicesForGroup;
 
 		// The set of base modifiers that are not currently represented within an up to date compiled section
 		TMap<FSoftObjectPath, MeshPartition::FModifierDesc> BaseModifiersNeedingCreation;
 
+		// Resolved grid cell size for this build variant (0 if not grid-aligned or resolution failed)
+		uint32 ResolvedGridCellSize = 0;
+
 	public:
 		// Returns the target group index for the group with the specified modifier set hash.  Returns -1 if no matching target group is found.
 		int32 GetTargetGroupIndexFromModifierSetHash(const FGuid& ModifierSetHash) const { return ModifierSetHashToTargetGroupIndex.FindRef(ModifierSetHash, -1); }
@@ -64,8 +69,10 @@ namespace UE::MeshPartition
 		const MeshPartition::FModifierGrouping& GetTargetGrouping() const { return TargetGrouping; }
 		const FCompiledSectionBuildInfo& GetBuildInfo() const { return BuildInfo; }
 		const FSoftObjectPath& GetMeshPartitionPath() const { return BuildInfo.MegaMeshPath; }
+		uint32 GetResolvedGridCellSize() const { return ResolvedGridCellSize; }
 
-		int32 GetReusableCompiledSectionForGroup(int32 GroupIndex, const FMeshPartitionUpdater& MeshPartitionUpdater) const;
+		bool HasReusableCompiledSectionsForGroup(int32 InGroupIndex, const FMeshPartitionUpdater& InMeshPartitionUpdater) const;
+		const TMap<FIntVector, int32>& GetReusableCompiledSectionIndicesForGroup(int32 InGroupIndex) const { return ReusableCompiledSectionIndicesForGroup[InGroupIndex]; }
 
 		const TMap<FSoftObjectPath, MeshPartition::FModifierDesc>& GetBaseModifiersNeedingCreation() const { return BaseModifiersNeedingCreation; }
 	};
@@ -96,7 +103,7 @@ namespace UE::MeshPartition
 		mutable int32 RemovedUnknown = 0;						// compiled sections removed because they did not correspond to any compiled section tracked in this updater (unexpected state)
 
 	public:
-		FMeshPartitionUpdater(const FGuid& InBuildKey, const UMeshPartitionDefinition& InMeshPartitionDefinition, TSharedPtr<MeshPartition::FCachedDescriptors> InMeshPartitionAllCachedDescriptors, TConstArrayView<FName> InTargetBuildVariantNames, IAssetRegistry& InAssetRegistry);
+		FMeshPartitionUpdater(const FGuid& InBuildKey, const UMeshPartitionDefinition& InMeshPartitionDefinition, TSharedPtr<MeshPartition::FCachedDescriptors> InMeshPartitionAllCachedDescriptors, TConstArrayView<FName> InTargetBuildVariantNames, IAssetRegistry& InAssetRegistry, const UWorld* InWorld);
 		const FBuildVariantUpdater* GetBuildVariantUpdater(const FName& BuildVariantName) const { return BuildVariantUpdaters.Find(BuildVariantName); }
 		void ForAllBuildVariants(TFunctionRef<bool(const MeshPartition::FBuildVariantUpdater&)> InFunc) const;
 
@@ -159,13 +166,14 @@ namespace UE::MeshPartition
 			// the counts here are mutually exclusive, even if a compiled section meets multiple criteria, it will only count towards one
 			// the exception is "Missing" which indicates compiled sections that do not exist (and can overlap with CannotFindFile)
 
-			// these compiled sections should be removed
+			// these compiled sections should be removed (RuntimeCellTransformer actively deletes them)
 			int32 Placeholders = 0;				// compiled sections that are PIE placeholder leftovers (not cleaned up from previous runs)
 			int32 NonTargetVariant = 0;			// compiled sections that do not correspond to any target build variant
 			int32 MismatchedModifiers = 0;		// compiled sections whose modifier sets do not correspond to any target group
 			int32 Duplicates = 0;				// compiled sections that are redundant duplicates of the same build variant + modifier group
+			int32 BuildVariantHashFails = 0;	// compiled sections whose build variant hash is out of date (definition/variant settings changed — section is fundamentally invalid)
 
-			// these compiled sections need to be updated
+			// these compiled sections need to be updated (content is stale but configuration matches)
 			int32 PackageHashFails = 0;			// compiled sections whose package hashes are out of date and need to be rebuilt, but do correspond to a current group
 			int32 ClassHashFails = 0;			// compiled sections whose class hashes are out of date and need to be rebuilt, but do correspond to a current group
 
@@ -183,16 +191,16 @@ namespace UE::MeshPartition
 		};
 		FCompiledSections CompiledSections;
 
-		// groups that have a corresponding existing compiled section that is up to date
-		int32 GroupsUpToDate() { return CompiledSections.Reuse; }
+		// Compiled sections that are up to date and will be reused (grid-split groups contribute N each)
+		int32 SectionsUpToDate() { return CompiledSections.Reuse; }
 
-		// groups that have a corresponding existing compiled section that is NOT up to date
-		int32 GroupsOutOfDate() { return CompiledSections.PackageHashFails + CompiledSections.ClassHashFails; }
+		// Compiled sections whose package or class hash is out of date
+		int32 SectionsOutOfDate() { return CompiledSections.PackageHashFails + CompiledSections.ClassHashFails; }
 
-		// groups that have NO corresponding existing compiled section
+		// Modifier groups that have no corresponding compiled section at all
 		int32 GroupsMissing() { return CompiledSections.Missing; }
 
 		// compiled sections that should be deleted
-		int32 CompiledSectionsToRemove() { return CompiledSections.Placeholders + CompiledSections.NonTargetVariant + CompiledSections.MismatchedModifiers + CompiledSections.Duplicates; }
+		int32 CompiledSectionsToRemove() { return CompiledSections.Placeholders + CompiledSections.NonTargetVariant + CompiledSections.MismatchedModifiers + CompiledSections.Duplicates + CompiledSections.BuildVariantHashFails; }
 	};
 } // namespace UE::MeshPartition
\ No newline at end of file
diff --git a/Engine/Plugins/Experimental/MeshPartition/Source/MeshPartitionEditor/Public/WorldPartitionMeshPartitionBuilder.h b/Engine/Plugins/Experimental/MeshPartition/Source/MeshPartitionEditor/Public/WorldPartitionMeshPartitionBuilder.h
--- a/Engine/Plugins/Experimental/MeshPartition/Source/MeshPartitionEditor/Public/WorldPartitionMeshPartitionBuilder.h
+++ b/Engine/Plugins/Experimental/MeshPartition/Source/MeshPartitionEditor/Public/WorldPartitionMeshPartitionBuilder.h
@@ -94,10 +94,16 @@ class UWorldPartitionMeshPartitionBuilder : public UWorldPartitionBuilder
 		// and deallocate it  before destroying the array this view references.
 		TArrayView<const FCompiledSectionDescriptor> OldCompiledSectionDescriptors;
 
+		// Resolved WP grid cell size (0 = no grid splitting). Set once per variant, before the task loop.
+		uint32								ResolvedGridCellSize = 0;
+
 		// Results
 		ETaskResult							Result = ETaskResult::None;
-		MeshPartition::ACompiledSection*			CompiledSection = nullptr;
-		FGuid								ReuseCompiledSectionActorDescGuid;
+
+		// All compiled sections produced for this group (N per WP grid cell; 1 for single-section path)
+		TArray<MeshPartition::ACompiledSection*>		CompiledSections;
+
+		TArray<FGuid>						ReuseCompiledSectionActorDescGuids;
 
 		// internal intermediate state
 		ETaskState							State = ETaskState::Invalid;
@@ -147,7 +153,7 @@ private:
 	/**
 	* Build a MeshPartition
 	*/
-	bool BuildMeshPartition(const AMeshPartition* MeshPartition, const FGuid& BuildKey, UE::MeshPartition::FModifierDescriptorCache& DescriptorCache, FPackageSourceControlHelper& InPackageHelper);
+	bool BuildMeshPartition(const AMeshPartition* InMeshPartition, const FGuid& BuildKey, UE::MeshPartition::FModifierDescriptorCache& DescriptorCache, FPackageSourceControlHelper& InPackageHelper);
 
 	/**
 	* Executes and manages all modifier group tasks upon completion.
