[Mesh Partition]

Display preview section context while editing modifier tool targets.

#Synced-CL 52208678

diff --git a/Engine/Plugins/Experimental/MeshPartition/Source/MeshPartition/Private/MeshPartitionChannel.cpp b/Engine/Plugins/Experimental/MeshPartition/Source/MeshPartition/Private/MeshPartitionChannel.cpp
--- a/Engine/Plugins/Experimental/MeshPartition/Source/MeshPartition/Private/MeshPartitionChannel.cpp
+++ b/Engine/Plugins/Experimental/MeshPartition/Source/MeshPartition/Private/MeshPartitionChannel.cpp
@@ -7,41 +7,9 @@
 
 namespace UE::MeshPartition
 {
-TArray<FName> FChannelMap::GetChannels() const
-{
-	TArray<FName> Channels;
-	for (const MeshPartition::FChannelDesc& Channel : ChannelDescs)
-	{
-		Channels.Add(Channel.Name);
-	}
-	return Channels;
-}
-
-int32 FChannelMap::FindChannel(FName InChannelName) const
-{
-	return ChannelDescs.IndexOfByKey(InChannelName);
-}
-
-void operator += (MeshPartition::IDependencyInterface& Dependencies, const MeshPartition::FChannelMap& ChannelMap)
-{
-	for (const MeshPartition::FChannelDesc& Channel : ChannelMap.GetChannelDescs())
-	{
-		Dependencies += Channel.Name;
-	}
-}
-
-bool FNameWrapper::SerializeFromMismatchedTag(const FPropertyTag& Tag, FStructuredArchive::FSlot Slot)
+namespace
 {
-	if (Tag.Type == NAME_NameProperty)
-	{
-		Slot << Name;
-		return true;
-	}
-
-	return false;
-}
-
-TFunction<bool(UPrimitiveComponent*)> FChannelPacking::MakeChannelTableSetterToSectionActorPrimitiveComponents(const TArray<uint8>& InChannelTable)
+void SetChannelTableOnPrimitiveComponent(UPrimitiveComponent* InPrimitiveComponent, TConstArrayView<uint8> InChannelTable)
 {
 	// Channel table is the channel index to texture slice indirection for the section
 	// unused channels (in this section) are marked -1
@@ -103,12 +71,54 @@ TFunction<bool(UPrimitiveComponent*)> FChannelPacking::MakeChannelTableSetterToS
 		break;
 	}
 
-	return [LambdaPackedChannelTable = MoveTemp(PackedChannelTable)](UPrimitiveComponent* PrimitiveComponent)
+	InPrimitiveComponent->SetDefaultCustomPrimitiveDataFloatArray(UE::MeshPartition::PrimitiveDataIndex, MakeConstArrayView((float*)PackedChannelTable.GetData(), PackedChannelTable.Num()));
+}
+}
+
+TArray<FName> FChannelMap::GetChannels() const
+{
+	TArray<FName> Channels;
+	for (const MeshPartition::FChannelDesc& Channel : ChannelDescs)
+	{
+		Channels.Add(Channel.Name);
+	}
+	return Channels;
+}
+
+int32 FChannelMap::FindChannel(FName InChannelName) const
+{
+	return ChannelDescs.IndexOfByKey(InChannelName);
+}
+
+void operator += (MeshPartition::IDependencyInterface& Dependencies, const MeshPartition::FChannelMap& ChannelMap)
+{
+	for (const MeshPartition::FChannelDesc& Channel : ChannelMap.GetChannelDescs())
 	{
-		PrimitiveComponent->SetDefaultCustomPrimitiveDataFloatArray(UE::MeshPartition::PrimitiveDataIndex, MakeConstArrayView((float*)LambdaPackedChannelTable.GetData(), LambdaPackedChannelTable.Num()));
+		Dependencies += Channel.Name;
+	}
+}
 
+bool FNameWrapper::SerializeFromMismatchedTag(const FPropertyTag& Tag, FStructuredArchive::FSlot Slot)
+{
+	if (Tag.Type == NAME_NameProperty)
+	{
+		Slot << Name;
 		return true;
-	};
+	}
+
+	return false;
 }
 
-} // namespace UE::MeshPartition
\ No newline at end of file
+void FChannelPacking::SetCustomPrimitiveData(UPrimitiveComponent* InPrimitiveComponent, TConstArrayView<uint8> InChannelTable, const FVector2f& InChannelTexcoordDesc)
+{
+	if (InPrimitiveComponent)
+	{
+		SetChannelTableOnPrimitiveComponent(InPrimitiveComponent, InChannelTable);
+
+		const int32 DataIndex = UE::MeshPartition::PrimitiveDataIndex + 4;
+		InPrimitiveComponent->SetDefaultCustomPrimitiveDataVector2f(DataIndex, InChannelTexcoordDesc);
+	}
+}
+
+
+} // namespace UE::MeshPartition
diff --git a/Engine/Plugins/Experimental/MeshPartition/Source/MeshPartition/Private/MeshPartitionCompiledSection.cpp b/Engine/Plugins/Experimental/MeshPartition/Source/MeshPartition/Private/MeshPartitionCompiledSection.cpp
--- a/Engine/Plugins/Experimental/MeshPartition/Source/MeshPartition/Private/MeshPartitionCompiledSection.cpp
+++ b/Engine/Plugins/Experimental/MeshPartition/Source/MeshPartition/Private/MeshPartitionCompiledSection.cpp
@@ -164,8 +164,7 @@ void ACompiledSection::UpdateVirtualTextureSettings()
 void ACompiledSection::SetupChannelDataOnChildPrimitiveComponents()
 {
 	// Reassign the custom primitive data parameters for this section used by the material shader to fetch channel texture
-	SetChannelTexcoordDesc(ChannelTexcoordDesc);
-	SetChannelTable(ChannelTable);
+	SetChannelData(ChannelTable, ChannelTexcoordDesc);
 }
 
 void ACompiledSection::SetMaterialCacheTileCount(const FIntPoint& TileCount)
@@ -324,10 +323,7 @@ void ACompiledSection::AddStaticMesh(UStaticMesh* InStaticMesh, const FStaticMes
 	// MeshComponent->MaterialCacheUVRegion = Descriptor.UVRegion;
 
 	// copy channel table to this mesh component
-	MeshComponent->SetDefaultCustomPrimitiveDataVector2f(UE::MeshPartition::PrimitiveDataIndex + 4, ChannelTexcoordDesc);
-	// Channel Table is assigned to all the primitive components of the Section as a custom primitive data used by the Material
-	auto ChannelTableSetter = FChannelPacking::MakeChannelTableSetterToSectionActorPrimitiveComponents(ChannelTable);
-	ChannelTableSetter(MeshComponent);
+	FChannelPacking::SetCustomPrimitiveData(MeshComponent, ChannelTable, ChannelTexcoordDesc);
 
 	MeshComponent->RegisterComponent();
 	MeshComponents.Emplace(MeshComponent);
@@ -348,7 +344,8 @@ void ACompiledSection::AddStaticMesh(UStaticMesh* InStaticMesh, const FStaticMes
 	VirtualTextureFallbackMeshComponent->SetVisibleInRayTracing(false);
 	VirtualTextureFallbackMeshComponent->SetCastShadow(false);
 
-	ChannelTableSetter(VirtualTextureFallbackMeshComponent);
+	FChannelPacking::SetCustomPrimitiveData(VirtualTextureFallbackMeshComponent, ChannelTable, ChannelTexcoordDesc);
+
 	VirtualTextureFallbackMeshComponent->RegisterComponent();
 
 	VirtualTextureFallbackMeshComponents.Emplace(VirtualTextureFallbackMeshComponent);
@@ -398,30 +395,20 @@ AMeshPartition* ACompiledSection::GetParentMegaMesh() const
 	return Parent.Get();
 }
 
-void ACompiledSection::SetChannelTexcoordDesc(const FVector2f& InTexcoordDesc)
-{
-	ChannelTexcoordDesc = InTexcoordDesc;
-
-	const int32 DataIndex = UE::MeshPartition::PrimitiveDataIndex + 4;
-
-	ForAllPrimitiveComponents([DataIndex, InTexcoordDesc](UPrimitiveComponent* PrimitiveComponent)
-	{
-		// StaticMesh component pass the texture slot as PrimitiveData
-		PrimitiveComponent->SetDefaultCustomPrimitiveDataVector2f(DataIndex, InTexcoordDesc);
-		return true;
-	});
-}
-
-void ACompiledSection::SetChannelTable(const TArray<uint8>& InChannelTable)
+void ACompiledSection::SetChannelData(TConstArrayView<uint8> InChannelTable, const FVector2f& InTexcoordDesc)
 {
 	// Channel table is the channel index to texture slice indirection for the section
 	// unused channels (in this section) are marked -1
 	// used channels have a valid index matching a slice of the section's channel texture array
 	ChannelTable = InChannelTable;
+	ChannelTexcoordDesc = InTexcoordDesc;
 
 	// Channel Table is assigned to all the primitive components of the Section as a custom primitive data used by the Material
-	auto ChannelTableSetter = FChannelPacking::MakeChannelTableSetterToSectionActorPrimitiveComponents(InChannelTable);
-	ForAllPrimitiveComponents(ChannelTableSetter);
+	ForAllPrimitiveComponents([&](UPrimitiveComponent* PrimitiveComponent)
+	{
+		FChannelPacking::SetCustomPrimitiveData(PrimitiveComponent, ChannelTable, ChannelTexcoordDesc);
+		return true;
+	});
 }
 
 void ACompiledSection::SetChannelTexture(UTexture* InChannelTexture)
diff --git a/Engine/Plugins/Experimental/MeshPartition/Source/MeshPartition/Public/MeshPartitionChannel.h b/Engine/Plugins/Experimental/MeshPartition/Source/MeshPartition/Public/MeshPartitionChannel.h
--- a/Engine/Plugins/Experimental/MeshPartition/Source/MeshPartition/Public/MeshPartitionChannel.h
+++ b/Engine/Plugins/Experimental/MeshPartition/Source/MeshPartition/Public/MeshPartitionChannel.h
@@ -122,11 +122,7 @@ struct FChannelPacking
 	//
 	// Practically speaking, with P = 4 and N = 5 we can represent 24 channels and thats our best pick for now
 
-	/**
-	 * Helper function used when assigning the Channel Table to a Section actor to pack the data according to the material expression unpacking logic.
-	 * and assign the packed data as custom primitive data on the actor's components
-	 */
-	static UE_API TFunction<bool(UPrimitiveComponent*)> MakeChannelTableSetterToSectionActorPrimitiveComponents(const TArray<uint8>& InChannelTable);
+	static UE_API void SetCustomPrimitiveData(UPrimitiveComponent* InPrimitiveComponent, TConstArrayView<uint8> InChannelTable, const FVector2f& InChannelTexcoordDesc);
 };
 
 /**
diff --git a/Engine/Plugins/Experimental/MeshPartition/Source/MeshPartition/Public/MeshPartitionCompiledSection.h b/Engine/Plugins/Experimental/MeshPartition/Source/MeshPartition/Public/MeshPartitionCompiledSection.h
--- a/Engine/Plugins/Experimental/MeshPartition/Source/MeshPartition/Public/MeshPartitionCompiledSection.h
+++ b/Engine/Plugins/Experimental/MeshPartition/Source/MeshPartition/Public/MeshPartitionCompiledSection.h
@@ -204,8 +204,7 @@ public:
 
 	void SetMaterialInstance(UMaterialInstanceConstant* InMaterialInstance) { MaterialInstance = InMaterialInstance; }
 	UMaterialInstanceConstant* GetMaterialInstance() const { return MaterialInstance; }
-	UE_API void SetChannelTexcoordDesc(const FVector2f& InTexcoordDesc);
-	UE_API void SetChannelTable(const TArray<uint8>& InChannelTable);
+	UE_API void SetChannelData(TConstArrayView<uint8> InChannelTable, const FVector2f& InTexcoordDesc);
 	UE_API void SetChannelTexture(UTexture* InChannelTexture);
 	UTexture* GetChannelTexture() { return ChannelTexture; }
 	UE_API void UpdateVirtualTextureSettings();
diff --git a/Engine/Plugins/Experimental/MeshPartition/Source/MeshPartitionEditor/Private/MeshPartitionChannelCollection.cpp b/Engine/Plugins/Experimental/MeshPartition/Source/MeshPartitionEditor/Private/MeshPartitionChannelCollection.cpp
--- a/Engine/Plugins/Experimental/MeshPartition/Source/MeshPartitionEditor/Private/MeshPartitionChannelCollection.cpp
+++ b/Engine/Plugins/Experimental/MeshPartition/Source/MeshPartitionEditor/Private/MeshPartitionChannelCollection.cpp
@@ -46,8 +46,7 @@ namespace UE::MeshPartition
 void ApplyChannels(ACompiledSection* InSection, const FSectionChannels& InChannels, const float InMaterialCacheTexelSize)
 {
 	InSection->SetChannelTexture(InChannels.Texture.Get());
-	InSection->SetChannelTexcoordDesc(InChannels.TexcoordMetrics);
-	InSection->SetChannelTable(InChannels.Table);
+	InSection->SetChannelData(InChannels.Table, InChannels.TexcoordMetrics);
 
 	if (InChannels.Domain.AreaUV > 0)
 	{
diff --git a/Engine/Plugins/Experimental/MeshPartition/Source/MeshPartitionEditor/Private/MeshPartitionEditorComponent.cpp b/Engine/Plugins/Experimental/MeshPartition/Source/MeshPartitionEditor/Private/MeshPartitionEditorComponent.cpp
--- a/Engine/Plugins/Experimental/MeshPartition/Source/MeshPartitionEditor/Private/MeshPartitionEditorComponent.cpp
+++ b/Engine/Plugins/Experimental/MeshPartition/Source/MeshPartitionEditor/Private/MeshPartitionEditorComponent.cpp
@@ -1407,6 +1407,8 @@ void UMeshPartitionEditorComponent::UpdateLinks(const UE::MeshPartition::FModifi
 		return;
 	}
 
+	InPreviewSection->SetGroupRegistryKey(MeshPartition::FModifierGroupRegistry::GetGroupRegistryKey(InModifierGroup));
+
 	InModifierGroup.ForAllModifiers([InPreviewSection](MeshPartition::UModifierComponent* Modifier)
 	{
 		if (Modifier->IsBase())
@@ -1419,8 +1421,6 @@ void UMeshPartitionEditorComponent::UpdateLinks(const UE::MeshPartition::FModifi
 
 		return true;
 	});
-
-	InPreviewSection->SetGroupRegistryKey(MeshPartition::FModifierGroupRegistry::GetGroupRegistryKey(InModifierGroup));
 }
 
 void UMeshPartitionEditorComponent::GetModifiersToProcessForPreviews(TSet<MeshPartition::UModifierComponent*>& 		OutModifiers,
@@ -1738,8 +1738,7 @@ void UMeshPartitionEditorComponent::FinalizePreviewSectionBuilds()
 			if (SectionChannel.IsValid())
 			{
 				PreviewSection->SetChannelTexture(SectionChannel->Texture.Get());
-				PreviewSection->SetChannelTexcoordDesc(SectionChannel->TexcoordMetrics);
-				PreviewSection->SetChannelTable(SectionChannel->Table);
+				PreviewSection->SetChannelData(SectionChannel->Table, SectionChannel->TexcoordMetrics);
 
 				if (SectionChannel->Mesh.IsSet())
 				{
@@ -1762,8 +1761,8 @@ void UMeshPartitionEditorComponent::FinalizePreviewSectionBuilds()
 
 			PreviewSections.Emplace(PreviewSection);
 
-			UpdateLinks(ModifierGroup, PreviewSection);
 			GroupRegistry.CommitStagedGroup(MeshPartition::FModifierGroupRegistry::GetGroupRegistryKey(ModifierGroup));
+			UpdateLinks(ModifierGroup, PreviewSection);
 
 			TArray<TWeakObjectPtr<MeshPartition::UModifierComponent>> ModifierPtrs = ModifierGroup.AllResolvedModifierPtrs();
 
@@ -1938,6 +1937,11 @@ void UMeshPartitionEditorComponent::ResetGroupRegistry()
 	GroupRegistry.Reset();
 }
 
+const MeshPartition::FModifierGroup* UMeshPartitionEditorComponent::FindGroupInRegistry(const FGuid& InGroupKey) const
+{
+	return GroupRegistry.FindGroup(InGroupKey);
+}
+
 bool UMeshPartitionEditorComponent::ShouldRecomputeTangents() const
 {
 	return bIsTransformerPipelinePaused && CVarRecomputeTangents.GetValueOnAnyThread();
diff --git a/Engine/Plugins/Experimental/MeshPartition/Source/MeshPartitionEditor/Private/MeshPartitionMeshBuilder.cpp b/Engine/Plugins/Experimental/MeshPartition/Source/MeshPartitionEditor/Private/MeshPartitionMeshBuilder.cpp
--- a/Engine/Plugins/Experimental/MeshPartition/Source/MeshPartitionEditor/Private/MeshPartitionMeshBuilder.cpp
+++ b/Engine/Plugins/Experimental/MeshPartition/Source/MeshPartitionEditor/Private/MeshPartitionMeshBuilder.cpp
@@ -258,7 +258,7 @@ namespace UE::MeshPartition::BuildHelpers
 			});
 		}
 
-		void FilterOutOfBoundsTriangles(MeshPartition::FMeshData& InOutMesh, TConstArrayView<FBox> InFilterBounds, const FTransform& InTransform)
+		void FilterOutOfBoundsTriangles(MeshPartition::FMeshData& InOutMesh, TConstArrayView<Geometry::FOrientedBox3d> InFilterBounds, const FTransform& InTransform, MeshPartition::EFilterBoundsMode InFilterMode)
 		{
 			TArray<TArray<int>> TrianglesToRemove;
 
@@ -278,9 +278,10 @@ namespace UE::MeshPartition::BuildHelpers
 					for (int Index = 0; Index < 3; ++Index)
 					{
 						const FVector3d Vertex = InTransform.TransformPosition(InOutMesh.GetVertex(Triangle[Index]));
-						for (const FBox& Bounds : InFilterBounds)
+						for (const Geometry::FOrientedBox3d& Bounds : InFilterBounds)
 						{
-							if (!Bounds.IsInsideOrOn(Vertex))
+							if ((InFilterMode == MeshPartition::EFilterBoundsMode::Inclusive && !Bounds.Contains(Vertex)) ||
+								(InFilterMode == MeshPartition::EFilterBoundsMode::Exclusive && Bounds.Contains(Vertex)))
 							{
 								OutTrianglesToRemove.Add(TriangleID);
 								return;
@@ -428,7 +429,7 @@ namespace UE::MeshPartition::BuildHelpers
 	{
 		FBlake3 Hasher;
 
-		auto UpdateHashPODType = []<typename Type UE_REQUIRES(TIsPODType<Type>::Value)> (FBlake3& Hasher, const Type& InData)
+		auto UpdateHashPODType = []<typename Type> (FBlake3& Hasher, const Type& InData)
 		{
 			Hasher.Update(reinterpret_cast<const uint8*>(&InData), sizeof(Type));
 		};
@@ -446,10 +447,11 @@ namespace UE::MeshPartition::BuildHelpers
 		UpdateHashPODType(Hasher, MeshPartition::FMeshData::GetVersionKey());
 
 		UpdateHashPODType(Hasher, InSettings.Transform);
-		for (const FBox& FilterBox : InSettings.FilterBounds)
+		for (const Geometry::FOrientedBox3d& FilterBox : InSettings.FilterBounds)
 		{
 			UpdateHashPODType(Hasher, FilterBox);
 		}
+		UpdateHashPODType(Hasher, InSettings.FilterBoundsMode);
 
 		if (InSettings.SimplifierOptions.IsSet())
 		{
@@ -1416,13 +1418,13 @@ MeshPartition::FBuildTaskHandle FMeshBuilder::Build_Internal(const MeshPartition
 		if (InSettings.FilterBounds.Num() > 0)
 		{
 			LastTaskInPipeline = Tasks::Launch(TEXT("ProcessModifiers_FilterOutOfBounds"),
-				[WeakResult = Result.ToWeakPtr(), FilterBounds = InSettings.FilterBounds, Transform = InSettings.Transform]
+				[WeakResult = Result.ToWeakPtr(), FilterBounds = InSettings.FilterBounds, FilterBoundsMode = InSettings.FilterBoundsMode, Transform = InSettings.Transform]
 				{
 					if (TSharedPtr<MeshPartition::FBuildTask> Result = WeakResult.Pin(); Result.IsValid() && !Result->bIsCancelled)
 					{
 						TSharedPtr<MeshPartition::FMeshData> ResultMesh = Result->Mesh;
 
-						FilterOutOfBoundsTriangles(*ResultMesh, FilterBounds, Transform);
+						FilterOutOfBoundsTriangles(*ResultMesh, FilterBounds, Transform, FilterBoundsMode);
 					}
 				},
 				Tasks::Prerequisites(LastTaskInPipeline)
diff --git a/Engine/Plugins/Experimental/MeshPartition/Source/MeshPartitionEditor/Private/MeshPartitionPreviewComponents.h b/Engine/Plugins/Experimental/MeshPartition/Source/MeshPartitionEditor/Private/MeshPartitionPreviewComponents.h
deleted file mode 100644
diff --git a/Engine/Plugins/Experimental/MeshPartition/Source/MeshPartitionEditor/Private/MeshPartitionPreviewSection.cpp b/Engine/Plugins/Experimental/MeshPartition/Source/MeshPartitionEditor/Private/MeshPartitionPreviewSection.cpp
--- a/Engine/Plugins/Experimental/MeshPartition/Source/MeshPartitionEditor/Private/MeshPartitionPreviewSection.cpp
+++ b/Engine/Plugins/Experimental/MeshPartition/Source/MeshPartitionEditor/Private/MeshPartitionPreviewSection.cpp
@@ -118,9 +118,8 @@ void APreviewSection::AddMesh(UStaticMesh* InStaticMesh, const FStaticMeshDescri
 	// TODO: Restore later
 	// MeshComponent->MaterialCacheUVRegion = Descriptor.UVRegion;
 	
-	// When creating these new MeshComponents, copy the CustomPrimitiveData from the PreviewComponent
-	const FCustomPrimitiveData& PrimData = PreviewMeshComponent->GetDefaultCustomPrimitiveData();
-	MeshComponent->SetDefaultCustomPrimitiveDataFloatArray(0, PrimData.Data);
+	// When creating these new MeshComponents, assign the custom primitive data required for rendering via the channel textures.
+	FChannelPacking::SetCustomPrimitiveData(MeshComponent, ChannelTable, ChannelTexcoordDesc);
 
 	MeshComponents.Emplace(MeshComponent);
 
@@ -149,7 +148,7 @@ void APreviewSection::AddMesh(UStaticMesh* InStaticMesh, const FStaticMeshDescri
 		VirtualTextureFallbackMeshComponent->SetCastShadow(false);
 		VirtualTextureFallbackMeshComponent->bSelectable = false;
 
-		VirtualTextureFallbackMeshComponent->SetDefaultCustomPrimitiveDataFloatArray(0, PrimData.Data);
+		FChannelPacking::SetCustomPrimitiveData(VirtualTextureFallbackMeshComponent, ChannelTable, ChannelTexcoordDesc);
 
 		VirtualTextureFallbackMeshComponent->RegisterComponent();
 
@@ -178,8 +177,7 @@ void APreviewSection::SetFarFieldMesh(UStaticMesh* InStaticMesh)
 
 	FarFieldMeshComponent->SetStaticMesh(InStaticMesh);
 
-	const FCustomPrimitiveData& PrimData = PreviewMeshComponent->GetDefaultCustomPrimitiveData();
-	FarFieldMeshComponent->SetDefaultCustomPrimitiveDataFloatArray(UE::MeshPartition::PrimitiveDataIndex, PrimData.Data);
+	FChannelPacking::SetCustomPrimitiveData(FarFieldMeshComponent, ChannelTable, ChannelTexcoordDesc);
 
 	InStaticMesh->OnPostMeshBuild().AddUObject(this, &APreviewSection::OnFarFieldStaticMeshBuild);
 }
@@ -259,30 +257,20 @@ void APreviewSection::SetMaterialInstance(UMaterialInstanceDynamic* InMaterialIn
 	});
 }
 
-void APreviewSection::SetChannelTexcoordDesc(const FVector2f& InTexcoordDesc)
-{
-	ChannelTexcoordDesc = InTexcoordDesc;
-
-	const int32 DataIndex = UE::MeshPartition::PrimitiveDataIndex + 4;
-
-	ForAllPrimitiveComponents([DataIndex, InTexcoordDesc](UPrimitiveComponent* PrimitiveComponent)
-	{
-		// StaticMesh component pass the texture slot as PrimitiveData
-		PrimitiveComponent->SetDefaultCustomPrimitiveDataVector2f(DataIndex, InTexcoordDesc);
-		return true;
-	});
-}
-
-void APreviewSection::SetChannelTable(const TArray<uint8>& InChannelTable)
+void APreviewSection::SetChannelData(TConstArrayView<uint8> InChannelTable, const FVector2f& InTexcoordDesc)
 {
 	// Channel table is the channel index to texture slice indirection for the section
 	// unused channels (in this section) are marked -1
 	// used channels have a valid index matching a slice of the section's channel texture array
 	ChannelTable = InChannelTable;
+	ChannelTexcoordDesc = InTexcoordDesc;
 
 	// Channel Table is assigned to all the primitive components of the Section as a custom primitive data used by the Material
-	auto ChannelTableSetter = FChannelPacking::MakeChannelTableSetterToSectionActorPrimitiveComponents(InChannelTable);
-	ForAllPrimitiveComponents(ChannelTableSetter);
+	ForAllPrimitiveComponents([&](UPrimitiveComponent* PrimitiveComponent)
+	{
+		FChannelPacking::SetCustomPrimitiveData(PrimitiveComponent, ChannelTable, ChannelTexcoordDesc);
+		return true;
+	});
 }
 
 void APreviewSection::SetChannelTexture(UTexture* InChannelTexture)
diff --git a/Engine/Plugins/Experimental/MeshPartition/Source/MeshPartitionEditor/Private/Modifiers/MeshPartitionProjectSculptLayersModifier.cpp b/Engine/Plugins/Experimental/MeshPartition/Source/MeshPartitionEditor/Private/Modifiers/MeshPartitionProjectSculptLayersModifier.cpp
--- a/Engine/Plugins/Experimental/MeshPartition/Source/MeshPartitionEditor/Private/Modifiers/MeshPartitionProjectSculptLayersModifier.cpp
+++ b/Engine/Plugins/Experimental/MeshPartition/Source/MeshPartitionEditor/Private/Modifiers/MeshPartitionProjectSculptLayersModifier.cpp
@@ -596,27 +596,7 @@ void UProjectMeshLayersModifier::PrepareForEdit(FDynamicMesh3& EditMesh) const
 
 	int32 StoredNumLayers = GetProjectionMeshNumLayers();
 
-	// filter EditMesh by edit bounds -- the mesh was pre-filtered by an AABB, but now we more-precisely filter by the oriented projection box
-	{
-		FTransform ProjectionTransform = GetProjectionTransform();
-		const FTransform& MeshToWorld = GetComponentTransform();
-		// matrix to take points from mesh space to projection transform space
-		FMatrix FromMeshToProjection = MeshToWorld.ToMatrixWithScale() * ProjectionTransform.ToInverseMatrixWithScale();
-		Geometry::FAxisAlignedBox3d ProjectionSpaceBounds = GetProjectionEditBounds();
-		for (int32 VID = 0, N = EditMesh.MaxVertexID(); VID < N; ++VID)
-		{
-			if (!EditMesh.IsVertex(VID))
-			{
-				continue;
-			}
-			FVector3d ProjPos = FromMeshToProjection.TransformPosition(EditMesh.GetVertex(VID));
-			if (!ProjectionSpaceBounds.Contains(ProjPos))
-			{
-				EditMesh.RemoveVertex(VID);
-			}
-		}
-		EditMesh.CompactInPlace();
-	}
+	EditMesh.CompactInPlace();
 
 	EditMesh.EnableAttributes();
 	
@@ -792,12 +772,12 @@ void UProjectMeshLayersModifier::PrepareForEdit(FDynamicMesh3& EditMesh) const
 	Geometry::FMeshNormals::QuickRecomputeOverlayNormals(EditMesh);
 }
 
-TArray<FBox> UProjectMeshLayersModifier::GetBoundsForEdit() const
+TArray<Geometry::FOrientedBox3d> UProjectMeshLayersModifier::GetBoundsForEdit() const
 {
 	FTransform ProjectionTransform = GetProjectionTransform();
 	Geometry::FAxisAlignedBox3d ProjectionSpaceBounds = GetProjectionEditBounds();
-	Geometry::FAxisAlignedBox3d GlobalBounds(ProjectionSpaceBounds, ProjectionTransform);
-	return { FBox(GlobalBounds) };
+	Geometry::FOrientedBox3d GlobalBounds(Geometry::FFrame3d(ProjectionTransform), ProjectionSpaceBounds.Extents() * ProjectionTransform.GetScale3D());
+	return { GlobalBounds };
 }
 
 TArray<FBox> UProjectMeshLayersModifier::ComputeBounds() const
diff --git a/Engine/Plugins/Experimental/MeshPartition/Source/MeshPartitionEditor/Public/MeshPartitionEditorComponent.h b/Engine/Plugins/Experimental/MeshPartition/Source/MeshPartitionEditor/Public/MeshPartitionEditorComponent.h
--- a/Engine/Plugins/Experimental/MeshPartition/Source/MeshPartitionEditor/Public/MeshPartitionEditorComponent.h
+++ b/Engine/Plugins/Experimental/MeshPartition/Source/MeshPartitionEditor/Public/MeshPartitionEditorComponent.h
@@ -382,6 +382,8 @@ public:
 	UE_API void ResetGroupRegistry(TConstArrayView<const MeshPartition::FModifierGroup> InModifierGroups);
 	UE_API void ResetGroupRegistry();
 
+	UE_API const MeshPartition::FModifierGroup* FindGroupInRegistry(const FGuid& InGroupKey) const;
+
 	/** Will return true if preview section transformer pipeline is paused and the recompute tangents feature is enabled. */
 	bool ShouldRecomputeTangents() const;
 
diff --git a/Engine/Plugins/Experimental/MeshPartition/Source/MeshPartitionEditor/Public/MeshPartitionGroupRegistry.h b/Engine/Plugins/Experimental/MeshPartition/Source/MeshPartitionEditor/Public/MeshPartitionGroupRegistry.h
--- a/Engine/Plugins/Experimental/MeshPartition/Source/MeshPartitionEditor/Public/MeshPartitionGroupRegistry.h
+++ b/Engine/Plugins/Experimental/MeshPartition/Source/MeshPartitionEditor/Public/MeshPartitionGroupRegistry.h
@@ -4,6 +4,8 @@
 
 #include "MeshPartitionModifierDescriptors.h"
 
+#define UE_API MESHPARTITIONEDITOR_API
+
 namespace UE::MeshPartition
 {
 
@@ -92,7 +94,7 @@ public:
 	 * @param InKey The group identifier.
 	 * @return Pointer to the group if found, otherwise nullptr.
 	 */
-	const MeshPartition::FModifierGroup* FindGroup(const FGuid& InKey) const;
+	UE_API const MeshPartition::FModifierGroup* FindGroup(const FGuid& InKey) const;
 
     /**
 	 * Finds a staged group by its key.
@@ -100,7 +102,7 @@ public:
 	 * @param InKey The group identifier.
 	 * @return Pointer to the staged group if found, otherwise nullptr.
 	 */
-	const MeshPartition::FModifierGroup* FindStagedGroup(const FGuid& InKey) const;
+	UE_API const MeshPartition::FModifierGroup* FindStagedGroup(const FGuid& InKey) const;
 
 	/**
 	 * Searches for groups containing the given set of modifiers.
@@ -108,7 +110,7 @@ public:
 	 * @param InModifiers List of modifiers to search for.
 	 * @return Array of groups that contain at least one of the specified modifiers.
 	 */
-	TArray<MeshPartition::FModifierGroup> FindGroupsContainingModifiers(TConstArrayView<MeshPartition::FModifierDesc> InModifiers) const;
+	UE_API TArray<MeshPartition::FModifierGroup> FindGroupsContainingModifiers(TConstArrayView<MeshPartition::FModifierDesc> InModifiers) const;
 	
 	/**
 	 * Searches for groups intersecting the provided bounds.
@@ -116,11 +118,13 @@ public:
 	 * @param InBounds List of bounds to check for.
 	 * @return Array of groups that contain at least one of the specified bounds.
 	 */
-	TArray<MeshPartition::FModifierGroup> FindGroupsIntersectingBounds(TConstArrayView<const FBox> InBounds) const;
+	UE_API TArray<MeshPartition::FModifierGroup> FindGroupsIntersectingBounds(TConstArrayView<const FBox> InBounds) const;
 
 private:
 	FGroupKeyToGroup Groups;
 	FGroupKeyToGroup StagedGroups;
 };
 
-}
\ No newline at end of file
+}
+
+#undef UE_API
\ No newline at end of file
diff --git a/Engine/Plugins/Experimental/MeshPartition/Source/MeshPartitionEditor/Public/MeshPartitionMeshBuilderCommon.h b/Engine/Plugins/Experimental/MeshPartition/Source/MeshPartitionEditor/Public/MeshPartitionMeshBuilderCommon.h
--- a/Engine/Plugins/Experimental/MeshPartition/Source/MeshPartitionEditor/Public/MeshPartitionMeshBuilderCommon.h
+++ b/Engine/Plugins/Experimental/MeshPartition/Source/MeshPartitionEditor/Public/MeshPartitionMeshBuilderCommon.h
@@ -10,6 +10,7 @@
 #include "MeshPartitionBuildPerfStats.h"
 #include "MeshPartitionChannelCollection.h"
 #include "MeshPartitionChannel.h"
+#include "OrientedBoxTypes.h"
 
 namespace UE::MeshPartition
 {
@@ -20,6 +21,14 @@ class UModifierComponent;
 
 using FModifierFilterFunc = TFunction<bool(const MeshPartition::FBuilderSettings&, const MeshPartition::FModifierDesc&)>;
 
+enum class EFilterBoundsMode
+{
+	/** Include triangles inside filter bounds. */
+	Inclusive,
+	/** Exclude triangles inside filter bounds. */
+	Exclusive,
+};
+
 struct FBuilderSettings
 {
 	/** The type of build requested. */
@@ -37,7 +46,9 @@ struct FBuilderSettings
 	* Optional bounds to filter the output mesh triangles.
 	* Only triangles fully contained within the bounds will be kept.
 	*/
-	TArray<FBox> FilterBounds;
+	TArray<Geometry::FOrientedBox3d> FilterBounds;
+	
+	MeshPartition::EFilterBoundsMode FilterBoundsMode = EFilterBoundsMode::Inclusive;
 
 	/** Optional filter function for the list of modifiers to process. */
 	FModifierFilterFunc ModifierFilter;
diff --git a/Engine/Plugins/Experimental/MeshPartition/Source/MeshPartitionEditor/Public/MeshPartitionModifierDescriptors.h b/Engine/Plugins/Experimental/MeshPartition/Source/MeshPartitionEditor/Public/MeshPartitionModifierDescriptors.h
--- a/Engine/Plugins/Experimental/MeshPartition/Source/MeshPartitionEditor/Public/MeshPartitionModifierDescriptors.h
+++ b/Engine/Plugins/Experimental/MeshPartition/Source/MeshPartitionEditor/Public/MeshPartitionModifierDescriptors.h
@@ -10,6 +10,8 @@
 
 #include "MeshPartitionModifierDescriptors.generated.h"
 
+#define UE_API MESHPARTITIONEDITOR_API
+
 class FWorldPartitionActorDescInstance;
 class UWorldPartition;
 class ULevelStreaming;
@@ -424,7 +426,7 @@ public:
 
 	void ForEachModifier(TFunctionRef<bool(MeshPartition::UModifierComponent*)> InFunc) const;
 
-	TArray<TWeakObjectPtr<MeshPartition::UModifierComponent>> AllResolvedModifierPtrs() const;
+	UE_API TArray<TWeakObjectPtr<MeshPartition::UModifierComponent>> AllResolvedModifierPtrs() const;
 
 private:
 	void ForEachModifier_Internal(TConstArrayView<TSoftObjectPtr<MeshPartition::UModifierComponent>> InModifierPtrs, TFunctionRef<bool(MeshPartition::UModifierComponent*)> InFunc) const;
@@ -484,7 +486,9 @@ private:
 	TOptional<MeshPartition::EBuildType> BuildType;
 };
 
-MESHPARTITIONEDITOR_API void SortModifierDescriptors(TConstArrayView<const FName> InModifierTypePriorities, TArray<MeshPartition::FModifierDesc>& InOutModifierDescriptors);
+UE_API void SortModifierDescriptors(TConstArrayView<const FName> InModifierTypePriorities, TArray<MeshPartition::FModifierDesc>& InOutModifierDescriptors);
 TArray<MeshPartition::FModifierGroup> BuildModifierGroups(TConstArrayView<MeshPartition::FModifierDesc> InModifierDescriptors, const MeshPartition::FBuilderSettings& InSettings);
 MeshPartition::FModifierGroup BuildModifierGroupForBase(TConstArrayView<MeshPartition::FModifierDesc> InBaseSet, TConstArrayView<MeshPartition::FModifierDesc> InAllNonBaseModifiersSorted);
 } // namespace UE::MeshPartition
+
+#undef UE_API
\ No newline at end of file
diff --git a/Engine/Plugins/Experimental/MeshPartition/Source/MeshPartitionEditor/Public/MeshPartitionPreviewComponents.h b/Engine/Plugins/Experimental/MeshPartition/Source/MeshPartitionEditor/Public/MeshPartitionPreviewComponents.h
new file mode 100644
--- /dev/null
+++ b/Engine/Plugins/Experimental/MeshPartition/Source/MeshPartitionEditor/Public/MeshPartitionPreviewComponents.h
@@ -0,0 +1,113 @@
+// Copyright Epic Games, Inc. All Rights Reserved.
+
+#pragma once
+
+#include "Components/MeshComponent.h"
+#include "Components/DynamicMeshComponent.h"
+#include "Components/StaticMeshComponent.h"
+#include "ComponentVisualizer.h"
+#include "MeshPartitionStaticMeshComponent.h"
+#include "MeshPartitionMeshData.h"
+
+#include "MeshPartitionPreviewComponents.generated.h"
+
+#define UE_API MESHPARTITIONEDITOR_API
+
+namespace UE::MeshPartition
+{
+/**
+* Subclass of static mesh component in order to customize the hit proxy, which needs to be custom to allow
+*  our visualizer to customize click handling for preview components (to route selection to base sections).
+*/
+UCLASS(MinimalAPI)
+class UStaticMeshPreviewComponent : public UMeshPartitionStaticMeshComponent
+{
+	GENERATED_BODY()
+
+public:
+	// Static mesh scene proxy calls into this method to set up hit proxies for different sections (note:
+	//  there is still also a hitproxy that points to the preview component itself, but that doesn't seem
+	//  to get in the way).
+	virtual HHitProxy* CreateMeshHitProxy(int32 SectionIndex, int32 MaterialIndex) const override;
+};
+
+
+UCLASS(MinimalAPI, meta = (BlueprintSpawnableComponent), ClassGroup = Rendering)
+class UPreviewMeshComponent : public UMeshComponent, public IInterface_CollisionDataProvider
+{
+	GENERATED_BODY()
+public:
+
+	UPreviewMeshComponent();
+	
+	virtual void GetResourceSizeEx(FResourceSizeEx& Size) override;
+
+	UE_API void SetMeshData(TSharedRef<const MeshPartition::FMeshData> InMeshData);
+	void SetMeshData(MeshPartition::FMeshData&& InMeshData);
+
+	TSharedPtr<const FMeshData> GetMeshData() { return MeshData; }
+	
+	virtual int32 GetNumMaterials() const override { return 1; }
+	virtual void SetMaterial(int32 ElementIndex, UMaterialInterface* Material) override;
+	virtual void GetUsedMaterials(TArray<UMaterialInterface*>& OutMaterials, bool bGetDebugMaterials) const override;
+	virtual UBodySetup* GetBodySetup() override;
+
+	virtual FBoxSphereBounds CalcBounds(const FTransform& LocalToWorld) const override;
+
+	virtual void OnRegister() override;
+	virtual void OnUnregister() override;
+
+	//~ Begin Interface_CollisionDataProvider Interface
+	virtual bool GetPhysicsTriMeshData(struct FTriMeshCollisionData* CollisionData, bool InUseAllTriData) override;
+	virtual bool ContainsPhysicsTriMeshData(bool InUseAllTriData) const override;
+	virtual bool GetTriMeshSizeEstimates(struct FTriMeshCollisionDataEstimates& OutTriMeshEstimates, bool bInUseAllTriData) const override;
+	virtual void GetMeshId(FString& OutMeshId) override;
+	//~ End Interface_CollisionDataProvider Interface
+
+private:
+	virtual FPrimitiveSceneProxy* CreateSceneProxy() override;
+
+	void RebuildPhysicsData();
+	UBodySetup* CreateBodySetupHelper();
+
+	void AbortPendingPhysicsCook();
+	void FinishPhysicsAsyncCook(bool bSuccess, UBodySetup* FinishedBodySetup);
+
+public:
+	UPROPERTY(Transient, Experimental)
+	TArray<TObjectPtr<class UMaterialCacheVirtualTexture>> MaterialCacheTextures;
+
+private:
+	UPROPERTY(Transient, Instanced)
+	TObjectPtr<UBodySetup> MeshBodySetup;
+
+	UPROPERTY(Transient, Instanced)
+	TObjectPtr<UBodySetup> AsyncBodySetup;  //Body setup being cooked asynchronously
+
+	TSharedPtr<const FMeshData> MeshData;
+};
+
+/**
+* Simple custom hitproxy similar to HActor, but inheriting from HComponentVisProxy to allow
+*  our visualizers to customize the click handling inside FComponentVisualizerManager::HandleProxyForComponentVis
+*/
+struct HPreviewProxy : public HComponentVisProxy
+{
+	DECLARE_HIT_PROXY();
+
+	HPreviewProxy(const UActorComponent* InComponent);
+
+	TWeakObjectPtr<const AActor> Actor;
+
+	virtual FTypedElementHandle GetElementHandle() const override
+	{
+		// The base class seems to crash in the element acquisition if the component is destroyed,
+		//  which seems to be possible.
+		return Component.IsValid() ? HComponentVisProxy::GetElementHandle()
+			: FTypedElementHandle();
+	}
+};
+
+} // namespace UE::MeshPartition
+
+#undef UE_API
\ No newline at end of file
diff --git a/Engine/Plugins/Experimental/MeshPartition/Source/MeshPartitionEditor/Public/MeshPartitionPreviewSection.h b/Engine/Plugins/Experimental/MeshPartition/Source/MeshPartitionEditor/Public/MeshPartitionPreviewSection.h
--- a/Engine/Plugins/Experimental/MeshPartition/Source/MeshPartitionEditor/Public/MeshPartitionPreviewSection.h
+++ b/Engine/Plugins/Experimental/MeshPartition/Source/MeshPartitionEditor/Public/MeshPartitionPreviewSection.h
@@ -59,7 +59,7 @@ public:
 	UE_API void AddCollisionComponent(UMeshPartitionCollisionComponent* InCollisionComponent);
 	UE_API TArray<TObjectPtr<UStaticMesh>> GetMeshes() const;
 	UE_API TSharedPtr<const MeshPartition::FMeshData> GetPreviewMesh() const;
-	UE_API UPreviewMeshComponent* GetPreviewMeshComponent() const { return PreviewMeshComponent; }
+	UPreviewMeshComponent* GetPreviewMeshComponent() const { return PreviewMeshComponent; }
 	UE_API void SetPreviewMesh(TSharedRef<const MeshPartition::FMeshData> InMeshData, const FName& InCollisionProfileName, const bool bInCanEverAffectNavigation);
 
 	UE_API AMeshPartition* GetParent() const;
@@ -68,20 +68,19 @@ public:
 	void SetMaterialInstance(UMaterialInstanceDynamic* InMaterialInstance);
 	UMaterialInstanceDynamic* GetMaterialInstance() const { return MaterialInstance; }
 
-	UE_API void SetChannelTexcoordDesc(const FVector2f& InTexcoorDesc);
-	UE_API void SetChannelTable(const TArray<uint8>& InChannelTable);
-	UE_API const TArray<uint8>& GetChannelTable() const { return ChannelTable; }
+	UE_API void SetChannelData(TConstArrayView<uint8> InChannelTable, const FVector2f& InTexcoordDesc);
+	const TArray<uint8>& GetChannelTable() const { return ChannelTable; }
 	UE_API void SetChannelTexture(UTexture* InChannelTexture);
-	UE_API UTexture* GetChannelTexture() const { return ChannelTexture; }
+	UTexture* GetChannelTexture() const { return ChannelTexture; }
 	
-	UE_API void SetChannelGenerationMesh(const MeshPartition::FSectionChannelGenerationMeshData& InSection) { SectionChannelMesh = InSection; }
-	UE_API const MeshPartition::FSectionChannelGenerationMeshData& GetChannelGenerationMesh() const { return SectionChannelMesh; }
+	void SetChannelGenerationMesh(const MeshPartition::FSectionChannelGenerationMeshData& InSection) { SectionChannelMesh = InSection; }
+	const MeshPartition::FSectionChannelGenerationMeshData& GetChannelGenerationMesh() const { return SectionChannelMesh; }
 
 	UE_API void SetBuildPerfStats(MeshPartition::FBuildPerfStats&& InBuildPerfStats);
 	UE_API const MeshPartition::FBuildPerfStats& GetBuildPerfStats() const;
 
-	UE_API void SetIsTransformerPipelinePaused(const bool bInIsPaused) { bIsTransformerPipelinePaused = bInIsPaused; }
-	UE_API bool IsTransformerPipelinedPaused() const { return bIsTransformerPipelinePaused;	}
+	void SetIsTransformerPipelinePaused(const bool bInIsPaused) { bIsTransformerPipelinePaused = bInIsPaused; }
+	bool IsTransformerPipelinedPaused() const { return bIsTransformerPipelinePaused;	}
 	
 	UE_API void OnMeshPartitionDefinitionChanged(const UMeshPartitionDefinition* InDefinition);
 	UE_API void OnMeshPartitionDefinitionModified(const UMeshPartitionDefinition* InDefinition, const FName& InPropertyName);
diff --git a/Engine/Plugins/Experimental/MeshPartition/Source/MeshPartitionEditor/Public/Modifiers/MeshPartitionEditableModifierBase.h b/Engine/Plugins/Experimental/MeshPartition/Source/MeshPartitionEditor/Public/Modifiers/MeshPartitionEditableModifierBase.h
--- a/Engine/Plugins/Experimental/MeshPartition/Source/MeshPartitionEditor/Public/Modifiers/MeshPartitionEditableModifierBase.h
+++ b/Engine/Plugins/Experimental/MeshPartition/Source/MeshPartitionEditor/Public/Modifiers/MeshPartitionEditableModifierBase.h
@@ -40,10 +40,19 @@ public:
 	}
 
 	// Bounds to use for editing the modifier.
+	// As opposed to ComputeBounds, these bounds are oriented to allow edited modifiers to more tightly fit their shape.
+	// This is especially important for the tools to generate a nice preview.
 	// Note: Is allowed to differ from ComputeBounds
-	virtual TArray<FBox> GetBoundsForEdit() const
+	virtual TArray<Geometry::FOrientedBox3d> GetBoundsForEdit() const
 	{
-		return ComputeBounds();
+		TArray<FBox> ModifierBounds = ComputeBounds();
+		TArray<Geometry::FOrientedBox3d> OrientedBounds;
+		Algo::Transform(ModifierBounds, OrientedBounds, [](const FBox& Box)
+			{
+				Geometry::FOrientedBox3d OrientedBox(Box); // Boxes returned by ComputeBounds are in world space already and orienting them would cause them to be larger than they should be.
+				return OrientedBox;
+			});
+		return OrientedBounds;
 	}
 
 
diff --git a/Engine/Plugins/Experimental/MeshPartition/Source/MeshPartitionEditor/Public/Modifiers/MeshPartitionProjectSculptLayersModifier.h b/Engine/Plugins/Experimental/MeshPartition/Source/MeshPartitionEditor/Public/Modifiers/MeshPartitionProjectSculptLayersModifier.h
--- a/Engine/Plugins/Experimental/MeshPartition/Source/MeshPartitionEditor/Public/Modifiers/MeshPartitionProjectSculptLayersModifier.h
+++ b/Engine/Plugins/Experimental/MeshPartition/Source/MeshPartitionEditor/Public/Modifiers/MeshPartitionProjectSculptLayersModifier.h
@@ -116,7 +116,7 @@ public:
 	// MeshPartition::UEditableModifierBase implementation
 	UE_API virtual void ApplyEditWithMesh(const FDynamicMesh3& UpdatedMesh) override;
 	UE_API virtual void PrepareForEdit(FDynamicMesh3& EditMesh) const override;
-	UE_API virtual TArray<FBox> GetBoundsForEdit() const override;
+	UE_API virtual TArray<Geometry::FOrientedBox3d> GetBoundsForEdit() const override;
 
 	// IMeshSculptLayersManager implementation
 	virtual bool HasSculptLayers() const override
diff --git a/Engine/Plugins/Experimental/MeshPartition/Source/MeshPartitionModelingToolset/Private/MeshPartitionModifierToolTarget.cpp b/Engine/Plugins/Experimental/MeshPartition/Source/MeshPartitionModelingToolset/Private/MeshPartitionModifierToolTarget.cpp
--- a/Engine/Plugins/Experimental/MeshPartition/Source/MeshPartitionModelingToolset/Private/MeshPartitionModifierToolTarget.cpp
+++ b/Engine/Plugins/Experimental/MeshPartition/Source/MeshPartitionModelingToolset/Private/MeshPartitionModifierToolTarget.cpp
@@ -18,6 +18,8 @@
 #include "DynamicMesh/DynamicMesh3.h"
 #include "DynamicMesh/MeshTransforms.h"
 #include "MeshPartitionMeshData.h"
+#include "Engine/World.h"
+#include "MeshPartitionToolPreviewActor.h"
 
 #include "DynamicMesh/Operations/MergeCoincidentMeshEdges.h"
 
@@ -77,19 +79,16 @@ void UModifierToolTarget::GetMaterialSet(FComponentMaterialSet& OutMaterialSet,
 
 bool UModifierToolTarget::CommitMaterialSetUpdate(const FComponentMaterialSet& InMaterialSet, bool bInApplyToAsset)
 {
-	// #todo_megamesh
 	return true;
 }
 
 UBodySetup* UModifierToolTarget::GetBodySetup() const
 {
-	// #todo_megamesh
 	return nullptr;
 }
 
 IInterface_CollisionDataProvider* UModifierToolTarget::GetComplexCollisionProvider() const
 {
-	// #todo_megamesh
 	return nullptr;
 }
 
@@ -97,11 +96,79 @@ void UModifierToolTarget::SetOwnerVisibility(bool bInVisible) const
 {
 	// TODO: This will need adjustment to handle "interactive section" visibility once we start using that again.
 
-	auto SetPreviewSectionVisibility = [bInVisible](MeshPartition::UMeshProviderModifier*, MeshPartition::APreviewSection* PreviewSection)
+	auto OnPreviewSectionReassigned = [this, WeakThis = TWeakObjectPtr<const UModifierToolTarget>(this), bInVisible](MeshPartition::UMeshProviderModifier* Modifier, MeshPartition::APreviewSection* PreviewSection)
 	{
-		if (PreviewSection && PreviewSection->GetRootComponent())
+		if (WeakThis.IsValid())
 		{
-			PreviewSection->GetRootComponent()->SetVisibility(bInVisible, /*bPropagateToChildren*/ true);
+			// When a preview section is reassigned, it by definition means that a prior preview section was destroyed.
+			// Find that (and any other) destroyed preview section and remove it and its preview:
+			for (auto It = PreviewSections.CreateIterator(); It; ++It)
+			{
+				TWeakObjectPtr<MeshPartition::APreviewSection> RegisteredPreview = *It;
+				if (!RegisteredPreview.IsValid())
+				{
+					if (TObjectPtr<AActor>* PreviewMesh = PreviewActors.Find(RegisteredPreview))
+					{
+						(*PreviewMesh)->GetWorld()->DestroyActor(*PreviewMesh, false);
+					}
+					It.RemoveCurrent();
+				}
+			}
+
+			if (PreviewSection && PreviewSection->GetRootComponent())
+			{
+				PreviewSection->GetRootComponent()->SetVisibility(bInVisible, /*bPropagateToChildren*/ true);
+
+				// If the new preview section was not already registered to the list of previews, add it now
+				PreviewSections.AddUnique(PreviewSection);
+
+				MeshPartition::UMeshPartitionEditorComponent* EditorComponent = PreviewSection->GetMegaMeshEditorComponent();
+
+				if (!PreviewActors.Contains(PreviewSection))
+				{
+					if (const MeshPartition::FModifierGroup* Group = EditorComponent->FindGroupInRegistry(PreviewSection->GetGroupRegistryKey()))
+					{
+						TArray<MeshPartition::UModifierComponent*> Modifiers;
+						Algo::Transform(Group->AllResolvedModifierPtrs(), Modifiers, [](const TWeakObjectPtr<UModifierComponent>& Modifier) { return Modifier.Get(); });
+
+						MeshPartition::FBuilderSettings BuilderSettings;
+						BuilderSettings.bCacheResult = true;
+		
+						const UMeshPartitionDefinition* Definition = EditorComponent->GetMegaMeshDefinition() ? EditorComponent->GetMegaMeshDefinition() : UMeshPartitionDefinition::GetDefaultMegaMeshDefinition();
+
+						BuilderSettings.ModifiersToProcess = MoveTemp(Modifiers);
+						BuilderSettings.Transform = EditorComponent->GetOwner()->GetTransform();
+						BuilderSettings.ModifierFilter = UE::MeshPartition::FilterHelpers::FilterModifiersByLastModifierToBuild(*TargetModifier, true);
+						BuilderSettings.FilterBounds = GetBounds();
+						BuilderSettings.FilterBoundsMode = MeshPartition::EFilterBoundsMode::Exclusive;
+						// Note: Large max section complexity to get a single contiguous/combined mesh
+						BuilderSettings.MaxSectionComplexity = FMathd::MaxReal;
+						BuilderSettings.TypePriorities = Definition->GetModifierTypePriorities();
+						BuilderSettings.TexcoordGenerationOptions = FChannelCollectionUVLayoutOptions::GetFromDefinition(Definition);
+						BuilderSettings.bRecomputeNormals = true;
+
+						FBuilderSettings::FChannelRenderSettings ChannelRenderSettings;
+						ChannelRenderSettings.ChannelMap = Definition->GetChannelMap();
+						ChannelRenderSettings.TexelSize = Definition->GetChannelTexelSize();
+						BuilderSettings.ChannelRenderSettings = ChannelRenderSettings;
+
+						TArray<MeshPartition::FBuildTaskHandle> BuildTaskHandles = UE::MeshPartition::Build::LaunchBuilds(BuilderSettings);
+
+						check(BuildTaskHandles.Num() == 1);
+
+						MeshPartition::FBuildTaskHandle& TaskHandle = BuildTaskHandles[0];
+
+						TSharedPtr<const FMeshData> BuiltMesh = TaskHandle.GetTask()->GetMesh();
+			
+						MeshPartition::AToolPreviewMesh* PreviewMesh = EditorComponent->GetWorld()->SpawnActor<MeshPartition::AToolPreviewMesh>(MeshPartition::AToolPreviewMesh::StaticClass(), EditorComponent->GetOwner()->GetTransform());
+						PreviewMesh->SetMesh(BuiltMesh.ToSharedRef());
+						PreviewMesh->SetMaterial(Definition->GetMaterial());
+						PreviewMesh->SetChannelData(*TaskHandle.GetTask()->GetSectionChannels());
+
+						PreviewActors.Add(PreviewSection, PreviewMesh);
+					}
+				}
+			}
 		}
 	};
 
@@ -109,15 +176,15 @@ void UModifierToolTarget::SetOwnerVisibility(bool bInVisible) const
 	{
 		// Gather all the base sections associated with our preview sections, and hide their preview sections in
 		//  a way that is robust to preview section reassignment.
-		for (TObjectPtr<MeshPartition::APreviewSection> PreviewSection : PreviewSections)
+		for (TWeakObjectPtr<MeshPartition::APreviewSection> PreviewSection : PreviewSections)
 		{
 			for (TWeakObjectPtr<MeshPartition::UModifierComponent> UncastBase : PreviewSection->GetBaseModifiers())
 			{
 				if (MeshPartition::UMeshProviderModifier* BaseSection = Cast<MeshPartition::UMeshProviderModifier>(UncastBase.Get()))
 				{
 					HiddenBaseSections.Add(BaseSection);
-					SetPreviewSectionVisibility(BaseSection, BaseSection->GetPreviewSection());
-					BaseSection->OnPreviewSectionReassignment().AddWeakLambda(this, SetPreviewSectionVisibility);
+					OnPreviewSectionReassigned(BaseSection, BaseSection->GetPreviewSection());
+					BaseSection->OnPreviewSectionReassignment().AddWeakLambda(this, OnPreviewSectionReassigned);
 				}
 			}
 		}
@@ -131,20 +198,22 @@ void UModifierToolTarget::SetOwnerVisibility(bool bInVisible) const
 			{
 				continue;
 			}
-			SetPreviewSectionVisibility(BaseSection, BaseSection->GetPreviewSection());
+			OnPreviewSectionReassigned(BaseSection, BaseSection->GetPreviewSection());
 			BaseSection->OnPreviewSectionReassignment().RemoveAll(this);
 		}
 		HiddenBaseSections.Reset();
+
+		for (const TPair<TWeakObjectPtr<MeshPartition::APreviewSection>, TObjectPtr<AActor>>& PreviewPair : PreviewActors)
+		{
+			PreviewPair.Value->GetWorld()->DestroyActor(PreviewPair.Value, false);
+		}
+		PreviewActors.Empty();
 	}
 }
 
 void UModifierToolTarget::ConfigurePreviewForRendering(UPrimitiveComponent* PrimitiveComponent) const
 {
-	TFunction<bool(UPrimitiveComponent*)> ChannelTableSetter = FChannelPacking::MakeChannelTableSetterToSectionActorPrimitiveComponents(ChannelTable);
-	ChannelTableSetter(PrimitiveComponent);
-
-	const int32 DataIndex = UE::MeshPartition::PrimitiveDataIndex + 4;
-	PrimitiveComponent->SetDefaultCustomPrimitiveDataVector2f(DataIndex, ChannelTexcoordDesc);
+	FChannelPacking::SetCustomPrimitiveData(PrimitiveComponent, ChannelTable, ChannelTexcoordDesc);
 }
 
 void UModifierToolTarget::UpdateRenderTextureForPreview(const FDynamicMesh3& InPreviewMesh)
@@ -200,7 +269,7 @@ void UModifierToolTarget::BuildModifiedMeshUpToTarget(FDynamicMesh3& OutResultMe
 		FTransformSRT3d EditorMeshTransform = (FTransformSRT3d)EditorComponent->GetOwner()->GetTransform();
 		FTransformSRT3d ModifierTransform = (FTransformSRT3d)GetWorldTransform();
 
-		TArray<FBox> ModifierBounds = GetBounds();
+		TArray<Geometry::FOrientedBox3d> ModifierBounds = GetBounds();
 		TArray<MeshPartition::UModifierComponent*> ModifiersToProcess = GetModifiersToProcess();
 
 		MeshPartition::FBuilderSettings BuilderSettings;
@@ -250,9 +319,16 @@ void UModifierToolTarget::BuildModifiedMeshUpToTarget(FDynamicMesh3& OutResultMe
 	}
 }
 
-TArray<FBox> UModifierToolTarget::GetBounds() const
+TArray<Geometry::FOrientedBox3d> UModifierToolTarget::GetBounds() const
 {
-	return TargetModifier->ComputeBounds();
+	TArray<FBox> ModifierBounds = TargetModifier->ComputeBounds();
+	TArray<Geometry::FOrientedBox3d> Bounds;
+	Algo::Transform(ModifierBounds, Bounds, [Transform = TargetModifier->GetComponentToWorld()](const FBox& LocalBox)
+		{
+			Geometry::FOrientedBox3d OrientedBox(LocalBox); // Boxes returned by ComputeBounds are in world space already and orienting them would cause them to be larger than they should be.
+			return OrientedBox;
+		});
+	return Bounds;
 }
 
 TArray<MeshPartition::UModifierComponent*> UModifierToolTarget::GetModifiersToProcess() const
@@ -297,7 +373,7 @@ void UEditableModifierToolTarget::CommitDynamicMesh(const FDynamicMesh3& CommitM
 	}
 }
 
-TArray<FBox> UEditableModifierToolTarget::GetBounds() const
+TArray<Geometry::FOrientedBox3d> UEditableModifierToolTarget::GetBounds() const
 {
 	if (MeshPartition::UEditableModifierBase* EditableModifier = Cast<MeshPartition::UEditableModifierBase>(TargetModifier))
 	{
@@ -315,7 +391,12 @@ TArray<MeshPartition::UModifierComponent*> UEditableModifierToolTarget::GetModif
 		MeshPartition::UEditableModifierBase* EditableModifier = Cast<MeshPartition::UEditableModifierBase>(TargetModifier);
 		if (ensure(EditableModifier))
 		{
-			EditorComponent->GetModifiersAffectingBounds(ModifiersToProcess, { EditableModifier->GetBoundsForEdit() });
+			TArray<FBox> AxisAlignedBoxes;
+			Algo::Transform(EditableModifier->GetBoundsForEdit(), AxisAlignedBoxes, [](const Geometry::FOrientedBox3d& OrientedBox)
+				{
+					return FBox(OrientedBox);
+				});
+			EditorComponent->GetModifiersAffectingBounds(ModifiersToProcess, AxisAlignedBoxes);
 		}
 	}
 	return ModifiersToProcess;
diff --git a/Engine/Plugins/Experimental/MeshPartition/Source/MeshPartitionModelingToolset/Private/MeshPartitionToolPreviewActor.cpp b/Engine/Plugins/Experimental/MeshPartition/Source/MeshPartitionModelingToolset/Private/MeshPartitionToolPreviewActor.cpp
new file mode 100644
--- /dev/null
+++ b/Engine/Plugins/Experimental/MeshPartition/Source/MeshPartitionModelingToolset/Private/MeshPartitionToolPreviewActor.cpp
@@ -0,0 +1,49 @@
+// Copyright Epic Games, Inc. All Rights Reserved.
+
+#include "MeshPartitionToolPreviewActor.h"
+#include "Engine/CollisionProfile.h"
+#include "MeshPartitionPreviewComponents.h"
+#include "MeshPartitionEditorUtils.h"
+#include "MeshPartitionChannelCollection.h"
+#include "Engine/Texture.h"
+#include "Materials/MaterialInstanceDynamic.h"
+#include "MeshPartitionDefinition.h"
+#include "MeshPartitionChannel.h"
+
+namespace UE::MeshPartition
+{
+
+AToolPreviewMesh::AToolPreviewMesh()
+{
+	bListedInSceneOutliner = false;
+
+	PreviewMeshComponent = CreateDefaultSubobject<MeshPartition::UPreviewMeshComponent>(TEXT("PreviewMeshComponent"));
+	RootComponent = PreviewMeshComponent;
+
+	PreviewMeshComponent->SetCollisionProfileName(UCollisionProfile::NoCollision_ProfileName);
+	PreviewMeshComponent->SetMobility(EComponentMobility::Static);
+}
+
+void AToolPreviewMesh::SetMesh(TSharedRef<const MeshPartition::FMeshData> MeshData)
+{
+	PreviewMeshComponent->SetMeshData(MeshData);
+}
+
+void AToolPreviewMesh::SetMaterial(UMaterialInterface* InMaterial)
+{
+	MID = MeshPartition::EditorUtils::GetOrCreateMaterialInstance(MID, InMaterial, this, TEXT("ModifierToolTargetMID"), RF_Transient);
+	PreviewMeshComponent->SetMaterial(0, MID);
+}
+
+void AToolPreviewMesh::SetChannelData(const MeshPartition::FSectionChannels& InChannelData)
+{
+	ChannelTable = InChannelData.Table;
+	ChannelTexture = InChannelData.Texture.Get();
+	ChannelTexcoordDesc = InChannelData.TexcoordMetrics;
+
+	MID->SetTextureParameterValue(UE::MeshPartition::ChannelTextureParameterName, ChannelTexture);
+	
+	FChannelPacking::SetCustomPrimitiveData(PreviewMeshComponent, ChannelTable, ChannelTexcoordDesc);
+}
+
+}
\ No newline at end of file
diff --git a/Engine/Plugins/Experimental/MeshPartition/Source/MeshPartitionModelingToolset/Private/MeshPartitionToolPreviewActor.h b/Engine/Plugins/Experimental/MeshPartition/Source/MeshPartitionModelingToolset/Private/MeshPartitionToolPreviewActor.h
new file mode 100644
--- /dev/null
+++ b/Engine/Plugins/Experimental/MeshPartition/Source/MeshPartitionModelingToolset/Private/MeshPartitionToolPreviewActor.h
@@ -0,0 +1,50 @@
+// Copyright Epic Games, Inc. All Rights Reserved.
+
+#pragma once
+
+#include "CoreMinimal.h"
+#include "GameFramework/Actor.h"
+
+#include "MeshPartitionToolPreviewActor.generated.h"
+
+#define UE_API MESHPARTITIONMODELINGTOOLSET_API
+
+class UMeshComponent;
+class UMaterialInterface;
+class UTexture;
+class UMaterialInstanceDynamic;
+
+namespace UE::MeshPartition
+{
+class FMeshData;
+class UPreviewMeshComponent;
+struct FSectionChannels;
+
+UCLASS(MinimalAPI, Transient)
+class AToolPreviewMesh : public AActor
+{
+	GENERATED_BODY()
+public:
+	AToolPreviewMesh();
+
+	UE_API void SetMesh(TSharedRef<const MeshPartition::FMeshData> InMeshData);
+	UE_API void SetChannelData(const MeshPartition::FSectionChannels& InChannelData);
+	UE_API void SetMaterial(UMaterialInterface* InMaterial);
+
+private:
+
+	UPROPERTY(Transient)
+	TObjectPtr<MeshPartition::UPreviewMeshComponent> PreviewMeshComponent;
+
+	UPROPERTY(Transient)
+	TObjectPtr<UTexture> ChannelTexture;
+
+	UPROPERTY(Transient)
+	TObjectPtr<UMaterialInstanceDynamic> MID;
+
+	TArray<uint8> ChannelTable;
+	FVector2f ChannelTexcoordDesc;
+};
+}
+
+#undef UE_API
\ No newline at end of file
diff --git a/Engine/Plugins/Experimental/MeshPartition/Source/MeshPartitionModelingToolset/Public/MeshPartitionModifierToolTarget.h b/Engine/Plugins/Experimental/MeshPartition/Source/MeshPartitionModelingToolset/Public/MeshPartitionModifierToolTarget.h
--- a/Engine/Plugins/Experimental/MeshPartition/Source/MeshPartitionModelingToolset/Public/MeshPartitionModifierToolTarget.h
+++ b/Engine/Plugins/Experimental/MeshPartition/Source/MeshPartitionModelingToolset/Public/MeshPartitionModifierToolTarget.h
@@ -13,6 +13,7 @@
 #include "UObject/Object.h"
 #include "MeshPartitionChannel.h"
 #include "MeshPartitionChannelCollection.h"
+#include "OrientedBoxTypes.h"
 
 #include "MeshPartitionModifierToolTarget.generated.h"
 
@@ -75,15 +76,14 @@ public:
 protected:
 	UE_API void BuildModifiedMeshUpToTarget(Geometry::FDynamicMesh3& OutResultMesh, bool bIncludeTargetModifier = false);
 
-	UE_API virtual TArray<FBox> GetBounds() const;
+	UE_API virtual TArray<Geometry::FOrientedBox3d> GetBounds() const;
 	UE_API virtual TArray<MeshPartition::UModifierComponent*> GetModifiersToProcess() const;
-
 protected:
 	UPROPERTY()
 	TObjectPtr<MeshPartition::UModifierComponent> TargetModifier = nullptr;
 	UPROPERTY()
-	TArray<TObjectPtr<MeshPartition::APreviewSection>> PreviewSections;
-	
+	mutable TArray<TWeakObjectPtr<MeshPartition::APreviewSection>> PreviewSections;
+
 	TArray<uint8> ChannelTable;
 	FVector2f ChannelTexcoordDesc;
 
@@ -92,9 +92,14 @@ protected:
 	UPROPERTY()
 	TObjectPtr<UMaterialInstanceDynamic> MID;
 
+	template <typename KeyType, typename ValueType>
+	using TWeakObjectPtrKeyMap = TMap<KeyType, ValueType, FDefaultSetAllocator, TWeakObjectPtrMapKeyFuncs<KeyType, ValueType>>;
+
 	/** Cached UV topology for fast texture updates during painting */
 	TOptional<MeshPartition::FChannelTextureRenderer::FChannelRenderUVMeshTopology> CachedUVTopology;
 
+	mutable TWeakObjectPtrKeyMap<TWeakObjectPtr<MeshPartition::APreviewSection>, TObjectPtr<AActor>> PreviewActors;
+
 private:
 	// When using SetOwnerVisibility, we hide all base sections that were entangled with our preview sections. This
 	//  keeps track of them so that we can unhide them.
@@ -121,7 +126,7 @@ public:
 	using IDynamicMeshCommitter::CommitDynamicMesh; // unhide the other overload
 
 protected:
-	UE_API virtual TArray<FBox> GetBounds() const override;
+	UE_API virtual TArray<Geometry::FOrientedBox3d> GetBounds() const override;
 	UE_API virtual TArray<MeshPartition::UModifierComponent*> GetModifiersToProcess() const override;
 };
 
