[Mesh Partition]

Changed how we identify base modifiers. Instead of interpreting the "Base" priority layer as bases, we instead rely on a virtual function.

#Synced-CL 52208678

diff --git a/Engine/Plugins/Experimental/MeshPartition/Source/MeshPartition/Public/MeshPartitionModule.h b/Engine/Plugins/Experimental/MeshPartition/Source/MeshPartition/Public/MeshPartitionModule.h
--- a/Engine/Plugins/Experimental/MeshPartition/Source/MeshPartition/Public/MeshPartitionModule.h
+++ b/Engine/Plugins/Experimental/MeshPartition/Source/MeshPartition/Public/MeshPartitionModule.h
@@ -33,6 +33,8 @@ struct FCustomVersion
 		AddedChecksumsToBuildInfo = 5,
 		// Added GridCellCoord (FIntVector) and ResolvedGridCellSize to FCompiledSectionBuildInfo for WP-grid-aligned sections
 		AddedGridCellCoordToBuildInfo = 6,
+		// Added a virtual method to identify modifiers as bases instead of relying on a named "Base" priority layer.
+		VirtualIsBaseModifier = 7,
 		// Add more versions as needed
 
 		LatestVersionPlusOne,
diff --git a/Engine/Plugins/Experimental/MeshPartition/Source/MeshPartitionEditor/Private/MeshPartitionEditorComponent.cpp b/Engine/Plugins/Experimental/MeshPartition/Source/MeshPartitionEditor/Private/MeshPartitionEditorComponent.cpp
--- a/Engine/Plugins/Experimental/MeshPartition/Source/MeshPartitionEditor/Private/MeshPartitionEditorComponent.cpp
+++ b/Engine/Plugins/Experimental/MeshPartition/Source/MeshPartitionEditor/Private/MeshPartitionEditorComponent.cpp
@@ -231,7 +231,7 @@ void UMeshPartitionEditorComponent::PostRegisterMegaMeshComponents()
 		UpdateMaterial();
 	}
 	
-	SetHiddenModifierTypes({MegaMeshModifierType::Base});
+	SetBaseModifiersHidden(true);
 	SetPreviewSectionsVisibility(true);
 
 	USelection::SelectionChangedEvent.AddUObject(this, &UMeshPartitionEditorComponent::OnSelectionChanged);
@@ -688,14 +688,13 @@ AActor* UMeshPartitionEditorComponent::SpawnBaseModifier(FDynamicMesh3&& InMesh,
 
 	MeshProvider->SetAffectedMeshPartition(MeshPartition);
 	MeshProvider->SetMesh(Forward<FDynamicMesh3&&>(InMesh));
-	MeshProvider->SetType(MegaMeshModifierType::Base);
 
 	for (int MaterialIndex = 0; MaterialIndex < InMaterials.Num(); ++MaterialIndex)
 	{
 		MeshProvider->SetMaterial(MaterialIndex, InMaterials[MaterialIndex]);
 	}
 
-	MeshProvider->SetIsTemporarilyHiddenInEditor(IsModifierHidden(MeshProvider));
+	MeshProvider->SetIsTemporarilyHiddenInEditor(ShouldHideBaseModifiers());
 	MeshProvider->SetIgnoreChanged(false);
 
 	NewActor->SetRootComponent(MeshProvider);
@@ -913,29 +912,18 @@ void UMeshPartitionEditorComponent::ForAllPreviewSections(TFunctionRef<bool(Mesh
 	}
 }
 
-void UMeshPartitionEditorComponent::SetHiddenModifierTypes(const TArray<FName>& InHiddenModifierTypes)
+void UMeshPartitionEditorComponent::SetBaseModifiersHidden(bool bInHideBaseModifiers)
 {
-	HiddenModifierTypes = InHiddenModifierTypes;
-	
-	for (MeshPartition::UModifierComponent* Modifier : CurrentModifiers)
+	bHideBaseModifiers = bInHideBaseModifiers;
+
+	ForAllCurrentModifiers([&](MeshPartition::UModifierComponent* Modifier)
 	{
-		if (Modifier == nullptr)
+		if (Modifier->IsBase())
 		{
-			continue;
+			Modifier->SetIsTemporarilyHiddenInEditor(ShouldHideBaseModifiers());
 		}
-
-		Modifier->SetIsTemporarilyHiddenInEditor(IsModifierHidden(Modifier));
-	}
-}
-
-bool UMeshPartitionEditorComponent::IsModifierHidden(const MeshPartition::UModifierComponent* InModifier) const
-{
-	if (InModifier == nullptr)
-	{
-		return false;
-	}
-	
-	return HiddenModifierTypes.Contains(InModifier->GetType());
+		return true;
+	});
 }
 
 uint32 UMeshPartitionEditorComponent::GetModifierLayerIndex(const FName& InModifierLayer) const
@@ -1045,7 +1033,10 @@ void UMeshPartitionEditorComponent::OnModifierLoaded(MeshPartition::UModifierCom
 	}
 
 	InModifier->InitializeModifier();
-	InModifier->SetIsTemporarilyHiddenInEditor(IsModifierHidden(InModifier));
+	if (InModifier->IsBase())
+	{
+		InModifier->SetIsTemporarilyHiddenInEditor(ShouldHideBaseModifiers());
+	}
 
 	if (IsInteractiveModeReady() && (LastActorSelected != nullptr))
 	{
diff --git a/Engine/Plugins/Experimental/MeshPartition/Source/MeshPartitionEditor/Private/MeshPartitionModifierComponent.cpp b/Engine/Plugins/Experimental/MeshPartition/Source/MeshPartitionEditor/Private/MeshPartitionModifierComponent.cpp
--- a/Engine/Plugins/Experimental/MeshPartition/Source/MeshPartitionEditor/Private/MeshPartitionModifierComponent.cpp
+++ b/Engine/Plugins/Experimental/MeshPartition/Source/MeshPartitionEditor/Private/MeshPartitionModifierComponent.cpp
@@ -370,7 +370,7 @@ void UModifierComponent::OnRegister()
 	AActor* Owner = GetOwner();
 	AMeshPartition* MeshPartition = GetAffectedMeshPartition();
 	
-	if ((GetType() == MegaMeshModifierType::Base) && (Owner != nullptr) && (MeshPartition != nullptr))
+	if ((IsBase()) && (Owner != nullptr) && (MeshPartition != nullptr))
 	{
 		Owner->AttachToActor(MeshPartition, FAttachmentTransformRules::KeepRelativeTransform);
 
diff --git a/Engine/Plugins/Experimental/MeshPartition/Source/MeshPartitionEditor/Private/MeshPartitionModifierDescriptors.cpp b/Engine/Plugins/Experimental/MeshPartition/Source/MeshPartitionEditor/Private/MeshPartitionModifierDescriptors.cpp
--- a/Engine/Plugins/Experimental/MeshPartition/Source/MeshPartitionEditor/Private/MeshPartitionModifierDescriptors.cpp
+++ b/Engine/Plugins/Experimental/MeshPartition/Source/MeshPartitionEditor/Private/MeshPartitionModifierDescriptors.cpp
@@ -69,6 +69,7 @@ FModifierDesc::FModifierDesc()
 	, ComplexityMultiplier(1.f)
 	, bIsContiguous(false)
 	, bIsDisabled(false)
+	, bIsBase(false)
 {
 }
 
@@ -86,6 +87,7 @@ FModifierDesc::FModifierDesc(const MeshPartition::UModifierComponent& InModifier
 	, ComplexityMultiplier(InModifier.GetComplexityMultiplier())
 	, bIsContiguous(InModifier.IsContiguous())
 	, bIsDisabled(InModifier.IsDisabled())
+	, bIsBase(InModifier.IsBase())
 {
 	if (const AActor* Owner = InModifier.GetOwner())
 	{
@@ -102,11 +104,6 @@ bool FModifierDesc::IsValid() const
 	return ModifierPath.IsValid();
 }
 
-bool FModifierDesc::IsBase() const
-{
-	return MeshPartition::UModifierComponent::IsBase(Type);
-}
-
 FModifierDesc::FModifierDesc(const FWorldPartitionActorDescInstance& InActorDescInstance, int32 InModifierIndex)
 	: MeshPartition::FModifierDesc()
 {
@@ -279,6 +276,7 @@ FModifierDesc::FModifierDesc(const FWorldPartitionActorDesc& InActorDesc, const
 	ComplexityMultiplier = InComponentDesc.ComplexityMultiplier;
 	bIsContiguous = InComponentDesc.bIsContiguous;
 	bIsDisabled = InComponentDesc.bIsDisabled;
+	bIsBase = InComponentDesc.bIsBase;
 }
 
 MeshPartition::FModifierGroup BuildModifierGroupForBase(TConstArrayView<MeshPartition::FModifierDesc> InBase, TConstArrayView<MeshPartition::FModifierDesc> InModifiers)
diff --git a/Engine/Plugins/Experimental/MeshPartition/Source/MeshPartitionEditor/Private/MeshPartitionModiiferComponentDesc.cpp b/Engine/Plugins/Experimental/MeshPartition/Source/MeshPartitionEditor/Private/MeshPartitionModiiferComponentDesc.cpp
--- a/Engine/Plugins/Experimental/MeshPartition/Source/MeshPartitionEditor/Private/MeshPartitionModiiferComponentDesc.cpp
+++ b/Engine/Plugins/Experimental/MeshPartition/Source/MeshPartitionEditor/Private/MeshPartitionModiiferComponentDesc.cpp
@@ -3,6 +3,7 @@
 #include "MeshPartitionModifierComponentDesc.h"
 #include "MeshPartitionModifierComponent.h"
 #include "WorldPartition/WorldPartitionActorDescArchive.h"
+#include "MeshPartitionModule.h"
 
 namespace UE::MeshPartition
 {
@@ -34,13 +35,15 @@ void FWorldPartitionModifierComponentDesc::Init(const UActorComponent* InCompone
 	ComplexityMultiplier = Desc.ComplexityMultiplier;
 	bIsContiguous = Desc.bIsContiguous;
 	bIsDisabled = Desc.bIsDisabled;
-
+	bIsBase = Desc.bIsBase;
 }
 
 void FWorldPartitionModifierComponentDesc::Serialize(FArchive& Ar)
 {
 	FWorldPartitionComponentDesc::Serialize(Ar);
 
+	Ar.UsingCustomVersion(MeshPartition::FCustomVersion::GUID);
+
 	Ar << ComponentTransform;
 	Ar << ModifierPath; // modifier path is technically serializable from the parent actor desc + component desc but there are edge cases and this is far simpler.
 	Ar << MegaMeshGuid;
@@ -53,6 +56,19 @@ void FWorldPartitionModifierComponentDesc::Serialize(FArchive& Ar)
 	Ar << TDeltaSerialize<float>(ComplexityMultiplier);
 	Ar << TDeltaSerialize<bool>(bIsContiguous);
 	Ar << TDeltaSerialize<bool>(bIsDisabled);
+		
+	// Before this version, modifier "base-ness" was defined by having a Type == "Base".
+	if (Ar.CustomVer(MeshPartition::FCustomVersion::GUID) >= MeshPartition::FCustomVersion::VirtualIsBaseModifier)
+	{
+		Ar << TDeltaSerialize<bool>(bIsBase);
+	}
+	else
+	{
+		if (Ar.IsLoading())
+		{
+			bIsBase = Type == TEXT("Base");
+		}
+	}
 }
 
 } // namespace UE::MeshPartition
diff --git a/Engine/Plugins/Experimental/MeshPartition/Source/MeshPartitionEditor/Public/MeshPartitionEditorComponent.h b/Engine/Plugins/Experimental/MeshPartition/Source/MeshPartitionEditor/Public/MeshPartitionEditorComponent.h
--- a/Engine/Plugins/Experimental/MeshPartition/Source/MeshPartitionEditor/Public/MeshPartitionEditorComponent.h
+++ b/Engine/Plugins/Experimental/MeshPartition/Source/MeshPartitionEditor/Public/MeshPartitionEditorComponent.h
@@ -289,19 +289,8 @@ public:
 
 	UE_API const MeshPartition::FCommonBuildVariant& GetPreviewBuildVariant() const;
 
-	/**
-	* Gets the types that should currently be hidden.
-	* @return An array containing the names of the hidden modifier types. 
-	*/
-	const TArray<FName>& GetHiddenModifierTypes() const { return HiddenModifierTypes; }
-	UE_API void SetHiddenModifierTypes(const TArray<FName>& InHiddenModifierTypes);
-
-	/**
-	* Returns true if the provided modifier is hidden.
-	* @param InModifier The modifier to test.
-	* @return True if the modifier is hidden.
-	*/
-	UE_API bool IsModifierHidden(const MeshPartition::UModifierComponent* InModifier) const;
+	UE_API void SetBaseModifiersHidden(bool bInHideBaseModifiers);
+	bool ShouldHideBaseModifiers() const { return bHideBaseModifiers; }
 
 	/**
 	* Returns the effective layer stack ordering index for a modifier's layer name, depending on if a definition is available or not.
@@ -488,9 +477,6 @@ private:
 	UPROPERTY(Transient)
 	TArray<TObjectPtr<MeshPartition::APreviewSection>> InteractivePreviewSections;
 
-	UPROPERTY(Transient)
-	TArray<FName> HiddenModifierTypes;
-
 	// As FTransformerContext can be passed to async thread and the array entry can be relocated, better allocate it.
 	TArray<TUniquePtr<MeshPartition::FTransformerContext>> PreviewSectionTransformerContexts;
 
@@ -502,6 +488,9 @@ private:
 	UPROPERTY(Transient)
 	TWeakObjectPtr<AActor> LastActorSelected = nullptr;
 
+	UPROPERTY(Transient)
+	bool bHideBaseModifiers = true;
+
 	UPROPERTY(Transient)
 	bool bArePreviewSectionsVisible = true;
 
diff --git a/Engine/Plugins/Experimental/MeshPartition/Source/MeshPartitionEditor/Public/MeshPartitionModifierComponent.h b/Engine/Plugins/Experimental/MeshPartition/Source/MeshPartitionEditor/Public/MeshPartitionModifierComponent.h
--- a/Engine/Plugins/Experimental/MeshPartition/Source/MeshPartitionEditor/Public/MeshPartitionModifierComponent.h
+++ b/Engine/Plugins/Experimental/MeshPartition/Source/MeshPartitionEditor/Public/MeshPartitionModifierComponent.h
@@ -395,11 +395,6 @@ public:
 	*/
 	UE_API bool IntersectsAnyBounds(TConstArrayView<FBox> InBoundsToTest) const;
 
-	/**
-	* @return True if the provided type will be the first to be processed on the stack.
-	*/
-	static bool IsBase(const FName& InType) { return InType == MegaMeshModifierType::Base; }
-
 	/**
 	* @return True if this modifier shouldn't be splitted across multiple sections.
 	*/
@@ -481,7 +476,7 @@ public:
 	* Allows to determine if this modifier will be used as a base.
 	* @return True if this modifier will be the first to be processed on the stack.
 	*/
-	virtual bool IsBase() const { return IsBase(Type); }
+	virtual bool IsBase() const { return false; }
 
 	/**
 	* Called when the parent AMeshPartition's UMeshPartitionDefinition changed.
diff --git a/Engine/Plugins/Experimental/MeshPartition/Source/MeshPartitionEditor/Public/MeshPartitionModifierComponentDesc.h b/Engine/Plugins/Experimental/MeshPartition/Source/MeshPartitionEditor/Public/MeshPartitionModifierComponentDesc.h
--- a/Engine/Plugins/Experimental/MeshPartition/Source/MeshPartitionEditor/Public/MeshPartitionModifierComponentDesc.h
+++ b/Engine/Plugins/Experimental/MeshPartition/Source/MeshPartitionEditor/Public/MeshPartitionModifierComponentDesc.h
@@ -40,6 +40,7 @@ public:
 	float ComplexityMultiplier;
 	bool bIsContiguous;
 	bool bIsDisabled;
+	bool bIsBase;
 };
 
 } // namespace UE::MeshPartition
\ No newline at end of file
diff --git a/Engine/Plugins/Experimental/MeshPartition/Source/MeshPartitionEditor/Public/MeshPartitionModifierDescriptors.h b/Engine/Plugins/Experimental/MeshPartition/Source/MeshPartitionEditor/Public/MeshPartitionModifierDescriptors.h
--- a/Engine/Plugins/Experimental/MeshPartition/Source/MeshPartitionEditor/Public/MeshPartitionModifierDescriptors.h
+++ b/Engine/Plugins/Experimental/MeshPartition/Source/MeshPartitionEditor/Public/MeshPartitionModifierDescriptors.h
@@ -110,6 +110,7 @@ struct FModifierDesc
 	float ComplexityMultiplier;
 	bool bIsContiguous;
 	bool bIsDisabled;
+	bool bIsBase;
 
 	// construct an empty / default descriptor
 	FModifierDesc();
@@ -125,7 +126,7 @@ struct FModifierDesc
 	bool operator==(const FModifierDesc&) const = default;
 
 	MESHPARTITIONEDITOR_API bool IsValid() const;
-	MESHPARTITIONEDITOR_API bool IsBase() const;
+	bool IsBase() const { return bIsBase; }
 };
 
 /** Typesafe wrapper for identifying and indexing modifier component data inside Modifier Groups */
diff --git a/Engine/Plugins/Experimental/MeshPartition/Source/MeshPartitionEditor/Public/Modifiers/MeshPartitionMeshProvider.h b/Engine/Plugins/Experimental/MeshPartition/Source/MeshPartitionEditor/Public/Modifiers/MeshPartitionMeshProvider.h
--- a/Engine/Plugins/Experimental/MeshPartition/Source/MeshPartitionEditor/Public/Modifiers/MeshPartitionMeshProvider.h
+++ b/Engine/Plugins/Experimental/MeshPartition/Source/MeshPartitionEditor/Public/Modifiers/MeshPartitionMeshProvider.h
@@ -65,6 +65,7 @@ public:
 	UE_API virtual double GetComplexity() const override;
 	UE_API virtual void SetIsTemporarilyHiddenInEditor(const bool bInIsHidden) override;
 	UE_API virtual void DrawVisualization(const FSceneView* View, FPrimitiveDrawInterface* PDI) const override;
+	virtual bool IsBase() const override { return true; }
 	// End MeshPartition::UModifierComponent Implementation
 	
 	/** Used to forward SetMaterial to the internal dynamic mesh. */
