FChannelTextureRenderer: wrap background thread UObject creation with GCScopeCuard

#Synced-CL 52307552

diff --git a/Engine/Plugins/Experimental/MeshPartition/Source/MeshPartitionEditor/Private/MeshPartitionChannelCollection.cpp b/Engine/Plugins/Experimental/MeshPartition/Source/MeshPartitionEditor/Private/MeshPartitionChannelCollection.cpp
--- a/Engine/Plugins/Experimental/MeshPartition/Source/MeshPartitionEditor/Private/MeshPartitionChannelCollection.cpp
+++ b/Engine/Plugins/Experimental/MeshPartition/Source/MeshPartitionEditor/Private/MeshPartitionChannelCollection.cpp
@@ -25,6 +25,7 @@
 #include "EngineUtils.h"
 #include "Editor.h"
 #include "GenerateMips.h"
+#include "UObject/GarbageCollection.h"
 
 #include "DynamicMesh/DynamicMeshTriangleAttribute.h"
 #include "Parameterization/DynamicMeshUVEditor.h"
@@ -493,6 +494,11 @@ UTexture* FChannelTextureRenderer::AllocateSectionTextureResource(int32 InSizeX,
 		(InSizeX % GPixelFormats[InFormat].BlockSizeX) == 0 &&
 		(InSizeY % GPixelFormats[InFormat].BlockSizeY) == 0)
 	{
+		// Prevent GC from running while creating UObjects off the game thread.
+		// UObjectBase::AddObject automatically sets EInternalObjectFlags::Async (a root flag)
+		// on objects created off the game thread, and SetRootFlags requires either the game
+		// thread or the GC lock to be held.
+		FGCScopeGuard GCScopeGuard;
 		SectionTexture = NewObject<UTexture2DArray>(InOuter, InName);
 
 		SectionTexture->SetPlatformData(new FTexturePlatformData());
