Mesh Paritition Spline Modifier: update the modifier after undoing a spline point transform

#Synced-CL 52353462

diff --git a/Engine/Plugins/Experimental/MeshPartition/Source/MeshPartitionEditor/Private/Modifiers/MeshPartitionSplineModifier.cpp b/Engine/Plugins/Experimental/MeshPartition/Source/MeshPartitionEditor/Private/Modifiers/MeshPartitionSplineModifier.cpp
--- a/Engine/Plugins/Experimental/MeshPartition/Source/MeshPartitionEditor/Private/Modifiers/MeshPartitionSplineModifier.cpp
+++ b/Engine/Plugins/Experimental/MeshPartition/Source/MeshPartitionEditor/Private/Modifiers/MeshPartitionSplineModifier.cpp
@@ -22,6 +22,7 @@
 #include "CleaningOps/RemeshMeshOp.h"
 #include "Curves/CurveFloat.h"
 #include "Kismet2/ComponentEditorUtils.h"
+#include "Misc/TransactionObjectEvent.h"
 #include "PrimitiveDrawInterface.h"
 #include "PrimitiveDrawingUtils.h"
 #include "SplineUtil.h"
@@ -97,6 +98,7 @@ USplineModifier::USplineModifier()
 
 USplineModifier::~USplineModifier()
 {
+	UnregisterSplineTransactedDelegate();
 }
 
 void USplineModifier::PreEditChange(FProperty* InPropertyAboutToChange)
@@ -162,11 +164,13 @@ void USplineModifier::InitializeModifier()
 
 	if (TemplateSplineComponentName.IsEmpty())
 	{
+		RegisterSplineTransactedDelegate();
 		return;
 	}
 
 	if (SplinePtr.IsValid() && SplinePtr->GetName() == TemplateSplineComponentName)
 	{
+		RegisterSplineTransactedDelegate();
 		return;
 	}
 
@@ -181,14 +185,19 @@ void USplineModifier::InitializeModifier()
 			if ((SplineComponent->GetName() == TemplateSplineComponentName) && (SplineComponent->GetOwner() == GetOwner()))
 			{
 				SetSplineComponent(SplineComponent);
+				// SetSplineComponent calls RegisterSplineTransactedDelegate
 				return;
 			}
 		}
 	}
+
+	RegisterSplineTransactedDelegate();
 }
 
 void USplineModifier::UninitializeModifier()
 {
+	UnregisterSplineTransactedDelegate();
+
 	if (EdgeFalloffCurve && EdgeFalloffUpdateHandle.IsValid())
 	{
 		EdgeFalloffCurve->OnUpdateCurve.Remove(EdgeFalloffUpdateHandle);
@@ -197,6 +206,36 @@ void USplineModifier::UninitializeModifier()
 	Super::UninitializeModifier();
 }
 
+void USplineModifier::RegisterSplineTransactedDelegate()
+{
+	UnregisterSplineTransactedDelegate();
+
+	const USplineComponent* const Spline = GetSplineComponent();
+	if (!Spline)
+	{
+		return;
+	}
+
+	OnSplineTransactedHandle = FCoreUObjectDelegates::OnObjectTransacted.AddLambda(
+		[this, Spline](UObject* Object, const FTransactionObjectEvent& Transaction)
+		{
+			if (Object == Spline && Transaction.GetEventType() == ETransactionObjectEventType::UndoRedo)
+			{
+				CurrentCachedData.Reset();
+				OnChanged(ComputeBounds(), EChangeType::StateChange);
+			}
+		});
+}
+
+void USplineModifier::UnregisterSplineTransactedDelegate()
+{
+	if (OnSplineTransactedHandle.IsValid())
+	{
+		FCoreUObjectDelegates::OnObjectTransacted.Remove(OnSplineTransactedHandle);
+		OnSplineTransactedHandle.Reset();
+	}
+}
+
 void USplineModifier::PostEditComponentMove(bool bFinished)
 {
 	// TODO: This callback is triggered when a spline on the actor has one of its vertices moved,
@@ -252,6 +291,8 @@ void USplineModifier::SetSplineComponent(USplineComponent* Spline, bool bUpdate)
 		// silent mode
 		SplinePtr = TSoftObjectPtr<USplineComponent>(FSoftObjectPath(Spline));
 	}
+
+	RegisterSplineTransactedDelegate();
 }
 
 void USplineModifier::BP_SetSplineComponent(USplineComponent* InSplineComponent)
diff --git a/Engine/Plugins/Experimental/MeshPartition/Source/MeshPartitionEditor/Public/Modifiers/MeshPartitionSplineModifier.h b/Engine/Plugins/Experimental/MeshPartition/Source/MeshPartitionEditor/Public/Modifiers/MeshPartitionSplineModifier.h
--- a/Engine/Plugins/Experimental/MeshPartition/Source/MeshPartitionEditor/Public/Modifiers/MeshPartitionSplineModifier.h
+++ b/Engine/Plugins/Experimental/MeshPartition/Source/MeshPartitionEditor/Public/Modifiers/MeshPartitionSplineModifier.h
@@ -310,6 +310,13 @@ private:
 	mutable FInterpCurveVector ProjectedSplineCurve;
 	mutable FCriticalSection ProjectedSplineCurveMutex;
 
+	/** Delegate handle for listening to undo/redo on the spline component */
+	FDelegateHandle OnSplineTransactedHandle;
+
+	/** Register/unregister the OnObjectTransacted listener for the spline component */
+	void RegisterSplineTransactedDelegate();
+	void UnregisterSplineTransactedDelegate();
+
 	friend class MegaMeshSplineModifierLocals::FBackgroundOp; // for FSharedData
 };
 } // namespace UE::MeshPartition
