MeshPartition USplineModifier: when dropping a new spline modifier onto the map, initialize the spline with two points

#Synced-CL 52356793

diff --git a/Engine/Plugins/Experimental/MeshPartition/Source/MeshPartitionEditor/Private/MeshPartitionModifierActor.cpp b/Engine/Plugins/Experimental/MeshPartition/Source/MeshPartitionEditor/Private/MeshPartitionModifierActor.cpp
--- a/Engine/Plugins/Experimental/MeshPartition/Source/MeshPartitionEditor/Private/MeshPartitionModifierActor.cpp
+++ b/Engine/Plugins/Experimental/MeshPartition/Source/MeshPartitionEditor/Private/MeshPartitionModifierActor.cpp
@@ -55,6 +55,7 @@ namespace ModifierActorLocals
 			constexpr float StartParam = 0.f;
 			constexpr float EndParam = 1.f;
 
+			Component.ClearSplinePoints(false);
 			Component.AddPoint(FSplinePoint(StartParam, StartPoint), false);
 			Component.AddPoint(FSplinePoint(EndParam, EndPoint), false);
 			Component.UpdateSpline();
diff --git a/Engine/Plugins/Experimental/MeshPartition/Source/MeshPartitionModelingToolset/Private/MeshPartitionPlaceModifierTool.cpp b/Engine/Plugins/Experimental/MeshPartition/Source/MeshPartitionModelingToolset/Private/MeshPartitionPlaceModifierTool.cpp
--- a/Engine/Plugins/Experimental/MeshPartition/Source/MeshPartitionModelingToolset/Private/MeshPartitionPlaceModifierTool.cpp
+++ b/Engine/Plugins/Experimental/MeshPartition/Source/MeshPartitionModelingToolset/Private/MeshPartitionPlaceModifierTool.cpp
@@ -265,6 +265,7 @@ void UPlaceModifierTool::SetupNewModifier(AActor* ContainingActor, MeshPartition
 				constexpr float StartParam = 0.f;
 				constexpr float EndParam = 1.f;
 
+				Component.ClearSplinePoints(false);
 				Component.AddPoint(FSplinePoint(StartParam, StartPoint), false);
 				Component.AddPoint(FSplinePoint(EndParam, EndPoint), false);
 				Component.UpdateSpline();
