Mesh Partition Spline Modifier: automatically expand the modifier bounds according to the max scale value on the spline
- also set default FalloffDistance to 50
- also set bDrawLocalBounds to True by default

#Synced-CL 52393202
#Synced-CL 52393202

diff --git a/Engine/Plugins/Experimental/MeshPartition/Source/MeshPartitionEditor/Private/Modifiers/MeshPartitionSplineModifier.cpp b/Engine/Plugins/Experimental/MeshPartition/Source/MeshPartitionEditor/Private/Modifiers/MeshPartitionSplineModifier.cpp
--- a/Engine/Plugins/Experimental/MeshPartition/Source/MeshPartitionEditor/Private/Modifiers/MeshPartitionSplineModifier.cpp
+++ b/Engine/Plugins/Experimental/MeshPartition/Source/MeshPartitionEditor/Private/Modifiers/MeshPartitionSplineModifier.cpp
@@ -804,6 +804,8 @@ void USplineModifier::GatherDependencies(MeshPartition::IDependencyInterface& De
 
 	Dependencies += FalloffDistance;
 	Dependencies += bUseSplineScaleForFalloff;
+	Dependencies += bExpandBoundsBySplineScale;
+	Dependencies += FalloffBoundsMultiplier;
 	Dependencies += bUseEdgeFalloffCurve;
 	if (EdgeFalloffCurve && bUseEdgeFalloffCurve)
 	{
@@ -854,14 +856,49 @@ FBox USplineModifier::ComputeLocalBounds() const
 			SplineTransform.TransformPosition(SplineBounds.GetCorner(i))));
 	}
 
+	float BoundsMultiplier = FalloffBoundsMultiplier;
+	if (bUseSplineScaleForFalloff && bExpandBoundsBySplineScale)
+	{
+		// Find a safe upper bound for Scale.Y along the spline. The Hermite curve is equivalent to a
+		// Bezier curve, which is bounded by its convex hull, so checking the Bezier control points for
+		// each segment guarantees we capture any overshoot from tangents.
+		const FSplineCurves& SplineCurves = SplineComponent->GetSplineCurves();
+		const TArray<FInterpCurvePoint<FVector>>& ScalePoints = SplineCurves.Scale.Points;
+		float MaxAbsScaleY = 0.0f;
+		for (int32 ScalePointIndex = 0; ScalePointIndex < ScalePoints.Num(); ++ScalePointIndex)
+		{
+			const FInterpCurvePoint<FVector>& CurrentPoint = ScalePoints[ScalePointIndex];
+			const float P0Y = CurrentPoint.OutVal.Y;
+			MaxAbsScaleY = FMath::Max(MaxAbsScaleY, FMath::Abs(P0Y));
+
+			// For Constant/Linear interp the curve can't overshoot the endpoint values,
+			// so the Bezier control point check is only needed for curve interp modes.
+			const bool bHasTangents = (CurrentPoint.InterpMode != CIM_Constant && CurrentPoint.InterpMode != CIM_Linear);
+			if (bHasTangents && ScalePointIndex + 1 < ScalePoints.Num())
+			{
+				const FInterpCurvePoint<FVector>& NextPoint = ScalePoints[ScalePointIndex + 1];
+				const float P1Y = NextPoint.OutVal.Y;
+				const float Diff = NextPoint.InVal - CurrentPoint.InVal;
+
+				// Hermite-to-Bezier: B1 = P0 + T0*Diff/3, B2 = P1 - T1*Diff/3
+				const float B1Y = P0Y + CurrentPoint.LeaveTangent.Y * Diff / 3.0f;
+				const float B2Y = P1Y - NextPoint.ArriveTangent.Y * Diff / 3.0f;
+
+				MaxAbsScaleY = FMath::Max(MaxAbsScaleY, FMath::Max(FMath::Abs(B1Y), FMath::Abs(B2Y)));
+			}
+		}
+		BoundsMultiplier = MaxAbsScaleY;
+	}
+
+	// If BoundsMultiplier is zero due to all scale points being zero this will return the bounds around the spline points
 	FBox ToReturn(LocalBounds);
 	if (bUseNearestSplineFrameForDisplacement && !ShouldFillInterior())
 	{
-		ToReturn = ToReturn.ExpandBy(FMath::Max(FalloffDistance * FalloffBoundsMultiplier, MaxProjectionHeightExtent));
+		ToReturn = ToReturn.ExpandBy(FMath::Max(FalloffDistance * BoundsMultiplier, MaxProjectionHeightExtent));
 	}
 	else
 	{
-		ToReturn = ToReturn.ExpandBy(FVector(FalloffDistance * FalloffBoundsMultiplier, FalloffDistance * FalloffBoundsMultiplier, MaxZDistance));
+		ToReturn = ToReturn.ExpandBy(FVector(FalloffDistance * BoundsMultiplier, FalloffDistance * BoundsMultiplier, MaxZDistance));
 	}
 
 	return ToReturn;
@@ -1115,4 +1152,14 @@ void USplineModifier::SetMeshClosedInterior(bool InbMeshClosedInterior)
 	bMeshClosedInterior = InbMeshClosedInterior;
 }
 
+void USplineModifier::SetUseSplineScaleForFalloff(bool bInUseSplineScaleForFalloff)
+{
+	bUseSplineScaleForFalloff = bInUseSplineScaleForFalloff;
+}
+
+void USplineModifier::SetExpandBoundsBySplineScale(bool bInExpandBoundsBySplineScale)
+{
+	bExpandBoundsBySplineScale = bInExpandBoundsBySplineScale;
+}
+
 } // namespace UE::MeshPartition
\ No newline at end of file
diff --git a/Engine/Plugins/Experimental/MeshPartition/Source/MeshPartitionEditor/Private/Test/MeshPartitionSplineModifierTests.cpp b/Engine/Plugins/Experimental/MeshPartition/Source/MeshPartitionEditor/Private/Test/MeshPartitionSplineModifierTests.cpp
--- a/Engine/Plugins/Experimental/MeshPartition/Source/MeshPartitionEditor/Private/Test/MeshPartitionSplineModifierTests.cpp
+++ b/Engine/Plugins/Experimental/MeshPartition/Source/MeshPartitionEditor/Private/Test/MeshPartitionSplineModifierTests.cpp
@@ -626,6 +626,129 @@ TEST_CLASS_WITH_FLAGS(SplineModifierTests, "MeshPartition.Modifier.Splines", EAu
 				});
 	}
 
+	TEST_METHOD(Verify_Bounds_Expand_By_Uniform_Spline_Scale)
+	{
+		const float UniformScaleY = 3.0f;
+		for (int32 PointIndex = 0; PointIndex < SplineComponent->GetNumberOfSplinePoints(); ++PointIndex)
+		{
+			SplineComponent->SetScaleAtSplinePoint(PointIndex, FVector(1.0f, UniformScaleY, 1.0f));
+		}
+		SplineComponent->UpdateSpline();
+
+		VerifyBoundsContainScaledFalloff(UniformScaleY);
+	}
+
+	TEST_METHOD(Verify_Bounds_Expand_By_Max_Spline_Scale)
+	{
+		// Set varying scale values — bounds should use the maximum.
+		SplineComponent->SetScaleAtSplinePoint(0, FVector(1.0f, 1.0f, 1.0f));
+		SplineComponent->SetScaleAtSplinePoint(1, FVector(1.0f, 5.0f, 1.0f));
+		SplineComponent->SetScaleAtSplinePoint(2, FVector(1.0f, 2.0f, 1.0f));
+		SplineComponent->UpdateSpline();
+
+		VerifyBoundsContainScaledFalloff(5.0f);
+	}
+
+	TEST_METHOD(Verify_Bounds_Account_For_Hermite_Tangent_Overshoot)
+	{
+		// Create a spline with a sharp step in scale: 1, 1, 10, 10, 1.
+		// The transition from 10 back to 1 at point 3->4 gives point 3 a large negative
+		// auto-tangent (from the asymmetry of its neighbors 10 and 1), which causes the
+		// Hermite interpolation to overshoot above 10 in the segment between points 2 and 3.
+		SplineComponent->ClearSplinePoints();
+		SplineComponent->AddSplinePoint(FVector(0.0f, 0.0f, 0.0f), ESplineCoordinateSpace::World, false);
+		SplineComponent->AddSplinePoint(FVector(100.0f, 0.0f, 0.0f), ESplineCoordinateSpace::World, false);
+		SplineComponent->AddSplinePoint(FVector(200.0f, 0.0f, 0.0f), ESplineCoordinateSpace::World, false);
+		SplineComponent->AddSplinePoint(FVector(300.0f, 0.0f, 0.0f), ESplineCoordinateSpace::World, false);
+		SplineComponent->AddSplinePoint(FVector(400.0f, 0.0f, 0.0f), ESplineCoordinateSpace::World, false);
+		SplineComponent->UpdateSpline();
+
+		SplineComponent->SetScaleAtSplinePoint(0, FVector(1.0f, 1.0f, 1.0f));
+		SplineComponent->SetScaleAtSplinePoint(1, FVector(1.0f, 1.0f, 1.0f));
+		SplineComponent->SetScaleAtSplinePoint(2, FVector(1.0f, 10.0f, 1.0f));
+		SplineComponent->SetScaleAtSplinePoint(3, FVector(1.0f, 10.0f, 1.0f));
+		SplineComponent->SetScaleAtSplinePoint(4, FVector(1.0f, 1.0f, 1.0f));
+		SplineComponent->UpdateSpline();
+
+		SplineModifier->SetUseSplineScaleForFalloff(true);
+		SplineModifier->SetExpandBoundsBySplineScale(true);
+		SplineModifier->UpdateSplineData();
+
+		const FSplineCurves SplineCurves = SplineComponent->GetSplineCurves();
+
+		// Get the max absolute scale value at the curve control points
+
+		float MaxControlPointScaleY = 0.0f;
+		for (const FInterpCurvePoint<FVector>& ScalePoint : SplineCurves.Scale.Points)
+		{
+			MaxControlPointScaleY = FMath::Max(MaxControlPointScaleY, FMath::Abs(ScalePoint.OutVal.Y));
+		}
+
+		// Sample the scale curve to verify that overshoot actually occurs — the interpolated
+		// scale should exceed the maximum control point value somewhere along the spline.
+
+		float MaxSampledScaleY = 0.0f;
+		const int32 NumSamples = 9;
+		for (int32 SampleIndex = 0; SampleIndex <= NumSamples; ++SampleIndex)
+		{
+			const float T = SplineCurves.Scale.Points[0].InVal +
+				(SplineCurves.Scale.Points.Last().InVal - SplineCurves.Scale.Points[0].InVal) * SampleIndex / NumSamples;
+			const float SampledScaleY = SplineCurves.Scale.Eval(T).Y;
+			MaxSampledScaleY = FMath::Max(MaxSampledScaleY, FMath::Abs(SampledScaleY));
+		}
+
+		// Verify the test setup actually produces overshoot — if not, the test is not meaningful.
+		ASSERT_THAT(IsTrue(MaxSampledScaleY > MaxControlPointScaleY,
+			*FString::Printf(TEXT("Test setup should produce interpolation overshoot: MaxSampledScaleY (%.2f) should be > MaxControlPointScaleY (%.2f)"),
+				MaxSampledScaleY, MaxControlPointScaleY)));
+
+		// Now verify the bounds account for it.
+		const TArray<FBox> Bounds = SplineModifier->ComputeBounds();
+		ASSERT_THAT(IsTrue(Bounds.Num() > 0, TEXT("ComputeBounds returned no bounds")));
+
+		const float FalloffDistance = SplineModifier->GetFalloffDistance();
+		const FVector Extent = Bounds[0].GetExtent();
+
+		ASSERT_THAT(IsTrue(Extent.Y > FalloffDistance * MaxControlPointScaleY,
+			*FString::Printf(TEXT("Bounds Y extent (%.2f) should be > FalloffDistance * MaxControlPointScaleY (%.2f) due to tangent overshoot"),
+				Extent.Y, FalloffDistance * MaxControlPointScaleY)));
+	}
+
+	TEST_METHOD(Verify_Bounds_Use_FalloffBoundsMultiplier_When_ExpandBySplineScale_Disabled)
+	{
+		// When bExpandBoundsBySplineScale is false, the bounds should use FalloffBoundsMultiplier
+		// instead of computing from spline scale.
+		const float MultiplierValue = 4.0f;
+
+		// Set a large spline scale that would produce bigger bounds if bExpandBoundsBySplineScale were true.
+		SplineComponent->SetScaleAtSplinePoint(0, FVector(1.0f, 8.0f, 1.0f));
+		SplineComponent->SetScaleAtSplinePoint(1, FVector(1.0f, 8.0f, 1.0f));
+		SplineComponent->SetScaleAtSplinePoint(2, FVector(1.0f, 8.0f, 1.0f));
+		SplineComponent->UpdateSpline();
+
+		SplineModifier->SetUseSplineScaleForFalloff(true);
+		SplineModifier->SetExpandBoundsBySplineScale(false);
+		SplineModifier->SetFalloffBoundsMultiplier(MultiplierValue);
+		SplineModifier->UpdateSplineData();
+
+		const TArray<FBox> Bounds = SplineModifier->ComputeBounds();
+		ASSERT_THAT(IsTrue(Bounds.Num() > 0, TEXT("ComputeBounds returned no bounds")));
+
+		const FVector Extent = Bounds[0].GetExtent();
+		const float FalloffDistance = SplineModifier->GetFalloffDistance();
+
+		// Bounds should reflect FalloffBoundsMultiplier (4.0), not the spline scale (8.0).
+		const float ExpectedExtent = FalloffDistance * MultiplierValue;
+		const float OverExpectedExtent = FalloffDistance * 8.0f;
+
+		ASSERT_THAT(IsTrue(Extent.Y >= ExpectedExtent,
+			*FString::Printf(TEXT("Bounds Y extent (%.2f) should be >= FalloffDistance * FalloffBoundsMultiplier = %.2f"),
+				Extent.Y, ExpectedExtent)));
+		ASSERT_THAT(IsTrue(Extent.Y < OverExpectedExtent,
+			*FString::Printf(TEXT("Bounds Y extent (%.2f) should be < FalloffDistance * SplineScale = %.2f (should not use spline scale)"),
+				Extent.Y, OverExpectedExtent)));
+	}
+
 	AFTER_EACH()
 	{
 		ScopedEditorWorld.Reset();
@@ -1126,6 +1249,31 @@ TEST_CLASS_WITH_FLAGS(SplineModifierTests, "MeshPartition.Modifier.Splines", EAu
 				});
 	}
 
+	/**
+	 * Verifies that ComputeBounds produces a box whose Y extent is at least
+	 * FalloffDistance * ExpectedMinMultiplier when bUseSplineScaleForFalloff and
+	 * bExpandBoundsBySplineScale are both enabled.
+	 */
+	void VerifyBoundsContainScaledFalloff(float ExpectedMinMultiplier)
+	{
+		SplineModifier->SetUseSplineScaleForFalloff(true);
+		SplineModifier->SetExpandBoundsBySplineScale(true);
+		SplineModifier->UpdateSplineData();
+
+		const TArray<FBox> Bounds = SplineModifier->ComputeBounds();
+		ASSERT_THAT(IsTrue(Bounds.Num() > 0, TEXT("ComputeBounds returned no bounds")));
+
+		const FVector Extent = Bounds[0].GetExtent();
+		const float FalloffDistance = SplineModifier->GetFalloffDistance();
+		const float ExpectedMinExtent = FalloffDistance * ExpectedMinMultiplier;
+
+		// The bounds extent in Y should be at least the scaled falloff distance.
+		// (The spline runs along X, so falloff expands primarily in Y.)
+		ASSERT_THAT(IsTrue(Extent.Y >= ExpectedMinExtent,
+			*FString::Printf(TEXT("Bounds Y extent (%.2f) should be >= FalloffDistance * %.2f = %.2f"),
+				Extent.Y, ExpectedMinMultiplier, ExpectedMinExtent)));
+	}
+
 	void CompareToReferenceAsset(const FString& TestName)
 	{
 		const FString ReferencePackagePath = ReferencePath / TestName;
diff --git a/Engine/Plugins/Experimental/MeshPartition/Source/MeshPartitionEditor/Public/Modifiers/MeshPartitionSplineModifier.h b/Engine/Plugins/Experimental/MeshPartition/Source/MeshPartitionEditor/Public/Modifiers/MeshPartitionSplineModifier.h
--- a/Engine/Plugins/Experimental/MeshPartition/Source/MeshPartitionEditor/Public/Modifiers/MeshPartitionSplineModifier.h
+++ b/Engine/Plugins/Experimental/MeshPartition/Source/MeshPartitionEditor/Public/Modifiers/MeshPartitionSplineModifier.h
@@ -150,6 +150,14 @@ public:
 
 	UE_API void SetMeshClosedInterior(bool InbMeshClosedInterior);
 
+	UE_API void SetUseSplineScaleForFalloff(bool bInUseSplineScaleForFalloff);
+
+	bool GetExpandBoundsBySplineScale() const
+	{
+		return bExpandBoundsBySplineScale;
+	}
+	UE_API void SetExpandBoundsBySplineScale(bool bInExpandBoundsBySplineScale);
+
 private:
 
 	friend class FSplineSoftObjectPointerDetails;
@@ -172,7 +180,7 @@ private:
 	* Determines how far in projected world distance the spline affects the mesh to either side of it. A higher value results in a slower slope.
 	*/
 	UPROPERTY(EditAnywhere, Category = "Coverage", meta = (UIMin = 5, UIMax = 1500, ClampMin = 0.1))
-	float FalloffDistance = 500.f;
+	float FalloffDistance = 50.f;
 
 	/**
 	* When true, the scale Y values along the spline are multiplied by FalloffDistance to get the final falloff distance to use.
@@ -227,11 +235,19 @@ private:
 	UPROPERTY(EditAnywhere, Category = "Coverage", meta = (EditCondition = "bUseNearestSplineFrameForDisplacement"))
 	float MaxProjectionHeightExtent = 500;
 
+	/**
+	* When true, the bounds expansion automatically accounts for the maximum spline scale value,
+	*  replacing FalloffBoundsMultiplier. When false, FalloffBoundsMultiplier is used instead.
+	*/
+	UPROPERTY(EditAnywhere, AdvancedDisplay, Category = "EdgeFalloff", meta = (EditCondition = "bUseSplineScaleForFalloff"))
+	bool bExpandBoundsBySplineScale = true;
+
 	/**
 	* When calculating bounds, the falloff value is multiplied by this value. This can be set above 1 if the spline has scale
 	*  above 1 along its length, and bUseSplineScaleForFalloff is true.
+	*  Ignored when bExpandBoundsBySplineScale is true.
 	*/
-	UPROPERTY(EditAnywhere, Category = "EdgeFalloff", meta = (UIMin = 0.1, UIMax = 10, ClampMin = 0.1))
+	UPROPERTY(EditAnywhere, AdvancedDisplay, Category = "EdgeFalloff", meta = (UIMin = 0.1, UIMax = 10, ClampMin = 0.1, EditCondition = "!bExpandBoundsBySplineScale"))
 	float FalloffBoundsMultiplier = 1;
 
 
@@ -239,7 +255,7 @@ private:
 	bool bDrawProjectionPlane = false;
 
 	UPROPERTY(EditAnywhere, Category = "Modifier|ModifierSettings", meta = (DisplayName = "Local Bounds"))
-	bool bDrawLocalBounds = false;
+	bool bDrawLocalBounds = true;
 
 	UPROPERTY(EditAnywhere, Category = "Modifier|ModifierSettings", meta = (DisplayName = "Spline Frames"))
 	bool bDrawSplineFrames = false;
