Decouple height smooth brush properties from general smooth brush properties, so they can have different defaults and different user-facing properties

#Synced-CL 52396359

diff --git a/Engine/Plugins/Experimental/MeshPartition/Source/MeshPartitionModelingToolset/Private/Sculpting/HeightBrushes.cpp b/Engine/Plugins/Experimental/MeshPartition/Source/MeshPartitionModelingToolset/Private/Sculpting/HeightBrushes.cpp
--- a/Engine/Plugins/Experimental/MeshPartition/Source/MeshPartitionModelingToolset/Private/Sculpting/HeightBrushes.cpp
+++ b/Engine/Plugins/Experimental/MeshPartition/Source/MeshPartitionModelingToolset/Private/Sculpting/HeightBrushes.cpp
@@ -326,7 +326,7 @@ void UE::Geometry::FHeightSmoothBrushOp::ApplyStamp(const FDynamicMesh3* Mesh,
 
 UClass* UE::Geometry::FHeightSmoothBrushOp::GetPropertiesClass(bool bPrimaryBrush)
 {
-	return bPrimaryBrush ? USmoothBrushOpProps::StaticClass() : USecondarySmoothBrushOpProps::StaticClass();
+	return bPrimaryBrush ? UHeightSmoothBrushOpProps::StaticClass() : USecondaryHeightSmoothBrushOpProps::StaticClass();
 }
 
 void UE::Geometry::FHeightFlattenBrushOp::BeginStroke(const FDynamicMesh3* Mesh, 
diff --git a/Engine/Plugins/Experimental/MeshPartition/Source/MeshPartitionModelingToolset/Private/Sculpting/HeightBrushes.h b/Engine/Plugins/Experimental/MeshPartition/Source/MeshPartitionModelingToolset/Private/Sculpting/HeightBrushes.h
--- a/Engine/Plugins/Experimental/MeshPartition/Source/MeshPartitionModelingToolset/Private/Sculpting/HeightBrushes.h
+++ b/Engine/Plugins/Experimental/MeshPartition/Source/MeshPartitionModelingToolset/Private/Sculpting/HeightBrushes.h
@@ -102,6 +102,62 @@ public:
 
 };
 
+UCLASS(MinimalAPI)
+class UHeightSmoothBrushOpProps : public UBaseSmoothBrushOpProps
+{
+	GENERATED_BODY()
+public:
+	/** Strength of the Brush */
+	UPROPERTY(EditAnywhere, Category = SmoothBrush, meta = (DisplayName = "Strength", UIMin = "0.0", UIMax = "1.0", ClampMin = "0.0", ClampMax = "1.0", ModelingQuickEdit, ModelingQuickSettings = 200))
+	float Strength = 0.5;
+
+	/** Amount of falloff to apply (0.0 - 1.0) */
+	UPROPERTY(EditAnywhere, Category = SmoothBrush, meta = (DisplayName = "Falloff", UIMin = "0.0", UIMax = "1.0", ClampMin = "0.0", ClampMax = "1.0"))
+	float Falloff = 0.5;
+
+	UPROPERTY(EditAnywhere, Category = SmoothBrush, meta = (ClampMin = 1, UIMax = 50))
+	int32 Iterations = 1;
+
+	/** If true, try to preserve the shape of the UV/3D mapping. This will limit Smoothing and Remeshing in some cases. */
+	UPROPERTY(EditAnywhere, Category = SmoothBrush)
+	bool bPreserveUVFlow = false;
+
+	virtual float GetStrength() override { return Strength; }
+	virtual void SetStrength(float NewStrength) override { Strength = FMathf::Clamp(NewStrength, 0.0f, 1.0f); }
+	virtual float GetFalloff() override { return Falloff; }
+	virtual bool GetPreserveUVFlow() override { return bPreserveUVFlow; }
+	virtual int32 GetIterations() { return Iterations; }
+};
+
+
+
+UCLASS(MinimalAPI)
+class USecondaryHeightSmoothBrushOpProps : public UBaseSmoothBrushOpProps
+{
+	GENERATED_BODY()
+public:
+	/** Strength of the Brush */
+	UPROPERTY(EditAnywhere, Category = ShiftSmoothBrush, meta = (DisplayName = "Strength", UIMin = "0.0", UIMax = "1.0", ClampMin = "0.0", ClampMax = "1.0"))
+	float Strength = 0.5;
+
+	/** Amount of falloff to apply */
+	UPROPERTY(EditAnywhere, Category = ShiftSmoothBrush, meta = (DisplayName = "Falloff", UIMin = "0.0", UIMax = "1.0", ClampMin = "0.0", ClampMax = "1.0"))
+	float Falloff = 0.5;
+
+	UPROPERTY(EditAnywhere, Category = ShiftSmoothBrush, meta = (ClampMin = 1, UIMax = 50))
+	int32 Iterations = 1;
+
+	/** If true, try to preserve the shape of the UV/3D mapping. This will limit Smoothing and Remeshing in some cases. */
+	UPROPERTY(EditAnywhere, Category = ShiftSmoothBrush)
+	bool bPreserveUVFlow = false;
+
+	virtual float GetStrength() override { return Strength; }
+	virtual void SetStrength(float NewStrength) override { Strength = FMathf::Clamp(NewStrength, 0.0f, 1.0f); }
+	virtual float GetFalloff() override { return Falloff; }
+	virtual bool GetPreserveUVFlow() override { return bPreserveUVFlow; }
+	virtual int32 GetIterations() { return Iterations; }
+};
+
 
 namespace UE::Geometry
 {
