Spline Modifier changes:
- ability to write to a weight channel
- "falloff plateau" property to enable flat regions around the spline

#Synced-CL 52634718

diff --git a/Engine/Plugins/Experimental/MeshPartition/Source/MeshPartitionEditor/Private/Modifiers/MeshPartitionSplineModifier.cpp b/Engine/Plugins/Experimental/MeshPartition/Source/MeshPartitionEditor/Private/Modifiers/MeshPartitionSplineModifier.cpp
--- a/Engine/Plugins/Experimental/MeshPartition/Source/MeshPartitionEditor/Private/Modifiers/MeshPartitionSplineModifier.cpp
+++ b/Engine/Plugins/Experimental/MeshPartition/Source/MeshPartitionEditor/Private/Modifiers/MeshPartitionSplineModifier.cpp
@@ -50,6 +50,7 @@ namespace MegaMeshSplineModifierLocals
 		MeshPartition::ESplineModifierInteriorSmoothMode InteriorSmoothMode;
 		
 		double FalloffDistance;
+		float FalloffPlateauRatio = 0.0f;
 		bool bUseSplineScaleForFalloff;
 		bool bUseEdgeFalloffCurve;
 		TUniquePtr<FRichCurve> EdgeFalloffCurve;
@@ -60,6 +61,16 @@ namespace MegaMeshSplineModifierLocals
 
 		ESplineModifierBlendMode BlendMode;
 
+		ESplineModifierWriteMode WriteMode = ESplineModifierWriteMode::Positions | ESplineModifierWriteMode::Weights;
+
+		struct FWeightEntry
+		{
+			FName WeightChannelName;
+			double Value;
+			ESplineWeightBlendMode BlendMode;
+		};
+		TArray<FWeightEntry> WeightEntries;
+
 		// When CachedDataIn is set, we want to use it. When ExternalCachedDataToInitialize is set, we want to
 		//  initialize it. Only one of them should be set.
 		TSharedPtr<const FSplineCachedSurfaceData> CachedDataIn;
@@ -74,7 +85,7 @@ namespace MegaMeshSplineModifierLocals
 		// Generate a new random guid before submitting any code changes to the op
 		static FGuid GetCodeVersionKey()
 		{
-			static FGuid VersionKey(TEXT("9fc99d42-69ee-430b-87ae-99a50e5e40c9"));
+			static FGuid VersionKey(TEXT("c8213f9d-693b-4987-9fd8-f040e8ec5bc2"));
 			return VersionKey;
 		}
 
@@ -90,6 +101,32 @@ namespace MegaMeshSplineModifierLocals
 		void ComputeLoopMesh(FSplineCachedSurfaceData& CachedData) const;
 		TUniquePtr<Geometry::FDynamicMeshOperator> CreateRemeshingOperator(FSplineCachedSurfaceData& CachedData, double TargetEdgeLength) const;
 	};
+
+	inline double BlendValues(
+		const double CurrentValue, const double NewValue, const double Alpha,
+		const MeshPartition::ESplineWeightBlendMode BlendMode)
+	{
+		double DesiredValue = NewValue;
+		switch (BlendMode)
+		{
+		case MeshPartition::ESplineWeightBlendMode::AlphaBlend:
+			// DesiredValue is already NewValue; Lerp(Current, New, Alpha) gives standard alpha blend.
+			break;
+		case MeshPartition::ESplineWeightBlendMode::Additive:
+			DesiredValue = CurrentValue + NewValue;
+			break;
+		case MeshPartition::ESplineWeightBlendMode::Min:
+			DesiredValue = FMath::Min(CurrentValue, NewValue);
+			break;
+		case MeshPartition::ESplineWeightBlendMode::Max:
+			DesiredValue = FMath::Max(CurrentValue, NewValue);
+			break;
+		default:
+			ensure(false);
+		}
+
+		return FMath::Lerp(CurrentValue, DesiredValue, Alpha);
+	}
 }
 
 USplineModifier::USplineModifier()
@@ -340,14 +377,37 @@ void USplineModifier::UpdateSplineData()
 void MegaMeshSplineModifierLocals::FBackgroundOp::GetInstancesInBounds(const FBox& InBounds, TArray<FInstanceInfo>& OutInstanceInfos) const
 {
 	AddDefaultInstanceIfIntersects(LocalBounds.TransformBy(ComponentTransform), InBounds, OutInstanceInfos);
+	if (OutInstanceInfos.IsEmpty())
+	{
+		return;
+	}
 
-	if (!OutInstanceInfos.IsEmpty() && ShouldFillInterior() && InteriorSmoothMode != MeshPartition::ESplineModifierInteriorSmoothMode::Simple)
+	// Read components
+	OutInstanceInfos[0].ReadViewComponents = EMeshViewComponents::VertexPos;
+	if (ShouldFillInterior() && InteriorSmoothMode != MeshPartition::ESplineModifierInteriorSmoothMode::Simple)
 	{
-		OutInstanceInfos[0].ReadViewComponents = EMeshViewComponents::VertexPos | EMeshViewComponents::DynamicSubmesh;
+		OutInstanceInfos[0].ReadViewComponents |= EMeshViewComponents::DynamicSubmesh;
 	}
-	else if (!OutInstanceInfos.IsEmpty() && bUseNearestSplineFrameForDisplacement && bProjectSplineToSurface)
+	else if (bUseNearestSplineFrameForDisplacement && bProjectSplineToSurface)
 	{
-		OutInstanceInfos[0].ReadViewComponents = EMeshViewComponents::VertexPos | EMeshViewComponents::DynamicSubmesh;
+		OutInstanceInfos[0].ReadViewComponents |= EMeshViewComponents::DynamicSubmesh;
+	}
+
+	// Write components
+	OutInstanceInfos[0].WriteViewComponents = EnumHasAnyFlags(WriteMode, ESplineModifierWriteMode::Positions) ? EMeshViewComponents::VertexPos : EMeshViewComponents::None;
+
+	// Weight channels
+	if (EnumHasAnyFlags(WriteMode, ESplineModifierWriteMode::Weights))
+	{
+		for (const FWeightEntry& Entry : WeightEntries)
+		{
+			if (!Entry.WeightChannelName.IsNone())
+			{
+				OutInstanceInfos[0].ReadViewComponents |= EMeshViewComponents::VertexAttributeWeight;
+				OutInstanceInfos[0].WriteViewComponents |= EMeshViewComponents::VertexAttributeWeight;
+				OutInstanceInfos[0].UsedChannels.Emplace(Entry.WeightChannelName);
+			}
+		}
 	}
 
 }
@@ -405,6 +465,7 @@ TSharedPtr<const MeshPartition::IModifierBackgroundOp> USplineModifier::CreateBa
 	Op->BlendMode = BlendMode;
 
 	Op->FalloffDistance = FalloffDistance;
+	Op->FalloffPlateauRatio = FalloffPlateauRatio;
 	Op->bUseSplineScaleForFalloff = bUseSplineScaleForFalloff;
 	Op->bUseEdgeFalloffCurve = bUseEdgeFalloffCurve;
 	if (EdgeFalloffCurve && bUseEdgeFalloffCurve)
@@ -412,7 +473,18 @@ TSharedPtr<const MeshPartition::IModifierBackgroundOp> USplineModifier::CreateBa
 		Op->EdgeFalloffCurve = MakeUnique<FRichCurve>(EdgeFalloffCurve->FloatCurve);
 	}
 	Op->bSplineIsLeftOfFalloffCurve = bSplineIsLeftOfFalloffCurve;
-	
+
+	Op->WriteMode = static_cast<ESplineModifierWriteMode>(WriteMode);
+
+	Op->WeightEntries.Reserve(WeightChannels.Num());
+	for (const FSplineModifierWeightEntry& WeightEntry : WeightChannels)
+	{
+		FBackgroundOp::FWeightEntry& OpEntry = Op->WeightEntries.AddDefaulted_GetRef();
+		OpEntry.WeightChannelName = WeightEntry.WeightChannelName.GetName();
+		OpEntry.Value = WeightEntry.Value;
+		OpEntry.BlendMode = WeightEntry.BlendMode;
+	}
+
 	Op->bMeshClosedInterior = bMeshClosedInterior;
 	// TODO: We shouldn't need to do this sampling here- we should be able to do it from the spline
 	//  curve data, but it currently happens to be easier to do while we still have the spline. 
@@ -602,21 +674,44 @@ void MegaMeshSplineModifierLocals::FBackgroundOp::ApplyModifications(MeshPartiti
 
 			if (TriangleID != FDynamicMesh3::InvalidID && ensure(CachedDataToUse->MeshedLoop->IsTriangle(TriangleID)))
 			{
-				FVector3d HitMeshPoint = CachedDataToUse->MeshedLoop->GetTriBaryPoint(TriangleID, BaryCoords.X, BaryCoords.Y, BaryCoords.Z);
-
-				const FVector3d VertToSurface = HitMeshPoint - SplineSpaceVertex;
-				const double SignedHeightToSpline = VertToSurface.Dot(SplineSpaceProjectionDirection);
-				if (BlendMode == ESplineModifierBlendMode::Min && SignedHeightToSpline > 0.0)
+				if (EnumHasAnyFlags(WriteMode, ESplineModifierWriteMode::Positions))
 				{
-					return;
+					const FVector3d HitMeshPoint = CachedDataToUse->MeshedLoop->GetTriBaryPoint(TriangleID, BaryCoords.X, BaryCoords.Y, BaryCoords.Z);
+					const FVector3d WorldHitPoint = SplineTransform.TransformPosition(HitMeshPoint);
+					const FVector3d VertToSurface = HitMeshPoint - SplineSpaceVertex;
+					const double SignedHeightToSpline = VertToSurface.Dot(SplineSpaceProjectionDirection);
+
+					bool bVertexClamped = false;
+					if (BlendMode == ESplineModifierBlendMode::Min && SignedHeightToSpline > 0.0)
+					{
+						bVertexClamped = true;
+					}
+					if (BlendMode == ESplineModifierBlendMode::Max && SignedHeightToSpline < 0.0)
+					{
+						bVertexClamped = true;
+					}
+
+					if (!bVertexClamped)
+					{
+						InMeshView.SetVertexPos(VertexIndex, InTransform.InverseTransformPosition(WorldHitPoint));
+					}
 				}
-				if (BlendMode == ESplineModifierBlendMode::Max && SignedHeightToSpline < 0.0)
+
+				// Interior vertices are at full influence
+				if (EnumHasAnyFlags(WriteMode, ESplineModifierWriteMode::Weights))
 				{
-					return;
+					for (const FWeightEntry& WeightEntry : WeightEntries)
+					{
+						if (WeightEntry.WeightChannelName.IsNone())
+						{
+							continue;
+						}
+						const double CurrentValue = InMeshView.GetVertexAttributeWeight(WeightEntry.WeightChannelName, VertexIndex);
+						const double BlendedValue = BlendValues(CurrentValue, WeightEntry.Value, 1.0, WeightEntry.BlendMode);
+						InMeshView.SetVertexAttributeWeight(WeightEntry.WeightChannelName, VertexIndex, BlendedValue);
+					}
 				}
 
-				FVector3d WorldHitPoint = SplineTransform.TransformPosition(HitMeshPoint);
-				InMeshView.SetVertexPos(VertexIndex, InTransform.InverseTransformPosition(WorldHitPoint));
 				return;
 			}
 
@@ -669,14 +764,17 @@ void MegaMeshSplineModifierLocals::FBackgroundOp::ApplyModifications(MeshPartiti
 		FVector3d VertToSpline = SplineWorldPoint - WorldVertex;
 
 		const double SignedHeightToSpline = VertToSpline.Dot(WorldProjectionDirectionToUse);
+
+		bool bVertexClamped = false;
 		if (BlendMode == ESplineModifierBlendMode::Min && SignedHeightToSpline > 0.0)
 		{
-			return;
+			bVertexClamped = true;
 		}
 		if (BlendMode == ESplineModifierBlendMode::Max && SignedHeightToSpline < 0.0)
 		{
-			return;
+			bVertexClamped = true;
 		}
+
 		double ProjectedDistanceToSpline = (VertToSpline - SignedHeightToSpline * WorldProjectionDirectionToUse).Length();
 
 		double FalloffToUse = FalloffDistance;
@@ -692,14 +790,22 @@ void MegaMeshSplineModifierLocals::FBackgroundOp::ApplyModifications(MeshPartiti
 			return;
 		}
 
+		const double PlateauDistance = FalloffToUse * FalloffPlateauRatio;
+		const double FalloffRegion = FalloffToUse - PlateauDistance;
+
 		float HeightAlpha;
-		if (bUseNearestSplineFrameForDisplacement && bProjectSplineToSurface)
+		if (FalloffRegion <= 0.0)
+		{
+			// Ratio is 1.0: hard edge, full influence everywhere within FalloffToUse
+			HeightAlpha = 1.0f;
+		}
+		else if (bUseNearestSplineFrameForDisplacement && bProjectSplineToSurface)
 		{
-			HeightAlpha = 1.0 - FMath::Clamp(ActualDistanceToProjectedSpline / FalloffToUse, 0.0, 1.0);
+			HeightAlpha = 1.0 - FMath::Clamp((ActualDistanceToProjectedSpline - PlateauDistance) / FalloffRegion, 0.0, 1.0);
 		}
 		else
 		{
-			HeightAlpha = 1.0 - FMath::Clamp(ProjectedDistanceToSpline / FalloffToUse, 0.0, 1.0);
+			HeightAlpha = 1.0 - FMath::Clamp((ProjectedDistanceToSpline - PlateauDistance) / FalloffRegion, 0.0, 1.0);
 		}
 
 		if (bUseEdgeFalloffCurve && EdgeFalloffCurve)
@@ -710,14 +816,31 @@ void MegaMeshSplineModifierLocals::FBackgroundOp::ApplyModifications(MeshPartiti
 			HeightAlpha = FMath::Clamp((CurveValue - CurveRange.X) / (CurveRange.Y - CurveRange.X), 0.0, 1.0);
 		}
 
-		FVector3d NewWorldPosition = WorldVertex + HeightAlpha * SignedHeightToSpline * WorldProjectionDirectionToUse;
+		if (EnumHasAnyFlags(WriteMode, ESplineModifierWriteMode::Positions) && !bVertexClamped)
+		{
+			FVector3d NewWorldPosition = WorldVertex + HeightAlpha * SignedHeightToSpline * WorldProjectionDirectionToUse;
 
-		// We could just use the instance desc bounds to clamp our output, but we'll be proper and use our local
-		//  bounds.
-		NewWorldPosition = ComponentTransform.TransformPosition(
-			LocalBounds.GetClosestPointTo(ComponentTransform.InverseTransformPosition(NewWorldPosition)));
-		
-		InMeshView.SetVertexPos(VertexIndex, InTransform.InverseTransformPosition(NewWorldPosition));
+			// We could just use the instance desc bounds to clamp our output, but we'll be proper and use our local
+			//  bounds.
+			NewWorldPosition = ComponentTransform.TransformPosition(
+				LocalBounds.GetClosestPointTo(ComponentTransform.InverseTransformPosition(NewWorldPosition)));
+
+			InMeshView.SetVertexPos(VertexIndex, InTransform.InverseTransformPosition(NewWorldPosition));
+		}
+
+		if (EnumHasAnyFlags(WriteMode, ESplineModifierWriteMode::Weights))
+		{
+			for (const FWeightEntry& WeightEntry : WeightEntries)
+			{
+				if (WeightEntry.WeightChannelName.IsNone())
+				{
+					continue;
+				}
+				const double CurrentValue = InMeshView.GetVertexAttributeWeight(WeightEntry.WeightChannelName, VertexIndex);
+				const double BlendedValue = BlendValues(CurrentValue, WeightEntry.Value, HeightAlpha, WeightEntry.BlendMode);
+				InMeshView.SetVertexAttributeWeight(WeightEntry.WeightChannelName, VertexIndex, BlendedValue);
+			}
+		}
 	});
 
 	// The smoother is used for detail transfer- it gets initialized with an unmodified mesh, gets everything outside
@@ -803,6 +926,7 @@ void USplineModifier::GatherDependencies(MeshPartition::IDependencyInterface& De
 	Dependencies += BlendMode;
 
 	Dependencies += FalloffDistance;
+	Dependencies += FalloffPlateauRatio;
 	Dependencies += bUseSplineScaleForFalloff;
 	Dependencies += bExpandBoundsBySplineScale;
 	Dependencies += FalloffBoundsMultiplier;
@@ -822,6 +946,15 @@ void USplineModifier::GatherDependencies(MeshPartition::IDependencyInterface& De
 	}
 	Dependencies += MeshedInteriorNumTriTarget;
 	Dependencies += InteriorSmoothMode;
+
+	Dependencies += WriteMode;
+
+	for (const FSplineModifierWeightEntry& WeightEntry : WeightChannels)
+	{
+		Dependencies += WeightEntry.WeightChannelName.GetName();
+		Dependencies += WeightEntry.Value;
+		Dependencies += WeightEntry.BlendMode;
+	}
 }
 
 TArray<FBox> USplineModifier::ComputeBounds() const
@@ -1105,6 +1238,11 @@ void USplineModifier::SetFalloffDistance(float InFalloffDistance)
 	FalloffDistance = InFalloffDistance;
 }
 
+void USplineModifier::SetFalloffPlateauRatio(float InFalloffPlateauRatio)
+{
+	FalloffPlateauRatio = FMath::Clamp(InFalloffPlateauRatio, 0.0f, 1.0f);
+}
+
 void USplineModifier::SetSplinePolygonErrorTolerance(float InSplinePolygonErrorTolerance)
 {
 	SplinePolygonErrorTolerance = InSplinePolygonErrorTolerance;
diff --git a/Engine/Plugins/Experimental/MeshPartition/Source/MeshPartitionEditor/Public/Modifiers/MeshPartitionSplineModifier.h b/Engine/Plugins/Experimental/MeshPartition/Source/MeshPartitionEditor/Public/Modifiers/MeshPartitionSplineModifier.h
--- a/Engine/Plugins/Experimental/MeshPartition/Source/MeshPartitionEditor/Public/Modifiers/MeshPartitionSplineModifier.h
+++ b/Engine/Plugins/Experimental/MeshPartition/Source/MeshPartitionEditor/Public/Modifiers/MeshPartitionSplineModifier.h
@@ -7,6 +7,7 @@
 #include "DynamicMesh/DynamicMesh3.h"
 #include "DynamicMesh/DynamicMeshAABBTree3.h"
 #include "Components/SplineComponent.h"
+#include "MeshPartitionChannel.h"
 #include "MeshPartitionModifierComponent.h"
 #include "Delegates/IDelegateInstance.h"
 #include "Polygon2.h"
@@ -48,6 +49,35 @@ enum class ESplineModifierBlendMode : uint8
 	Max
 };
 
+/** Which outputs a spline modifier writes. */
+UENUM(meta = (Bitflags, UseEnumValuesAsMaskValuesInEditor = "true"))
+enum class ESplineModifierWriteMode : uint8
+{
+	/** Displace vertex positions along the spline projection direction. */
+	Positions	= (1 << 0),
+
+	/** Write to weight channels using the spline falloff as the blending alpha. */
+	Weights		= (1 << 1)
+};
+ENUM_CLASS_FLAGS(ESplineModifierWriteMode);
+
+/** How a spline modifier's weight channel value is combined with existing values. */
+UENUM()
+enum class ESplineWeightBlendMode : uint8
+{
+	/** The new value replaces the existing value, blended by the spline falloff. */
+	AlphaBlend,
+
+	/** The new value is added to the existing value, scaled by the spline falloff. */
+	Additive,
+
+	/** Takes the minimum of the existing and new values, blended by the spline falloff. */
+	Min,
+
+	/** Takes the maximum of the existing and new values, blended by the spline falloff. */
+	Max
+};
+
 UENUM()
 enum class ESplineModifierInteriorSmoothMode : uint8
 {
@@ -61,6 +91,28 @@ enum class ESplineModifierInteriorSmoothMode : uint8
 	DetailPreserving = 3
 };
 
+/**
+* Describes a single weight channel entry for a spline modifier. The spline's own falloff (HeightAlpha)
+* is used as the blending alpha.
+*/
+USTRUCT()
+struct FSplineModifierWeightEntry
+{
+	GENERATED_BODY()
+
+	/** Which weight channel to write. */
+	UPROPERTY(EditAnywhere, Category = "", meta = (GetOptions = "GetMegaMeshDefinitionChannels", NoResetToDefault))
+	MeshPartition::FChannelName WeightChannelName;
+
+	/** The value to write at full influence. */
+	UPROPERTY(EditAnywhere, Category = "")
+	double Value = 1.0;
+
+	/** How the value is blended with the existing channel value. */
+	UPROPERTY(EditAnywhere, Category = "")
+	MeshPartition::ESplineWeightBlendMode BlendMode = MeshPartition::ESplineWeightBlendMode::AlphaBlend;
+};
+
 /**
 * A flexible spline based modifier, suitable for representing ridgelines, channels, raised hills, and lowered depression pockets.
 */
@@ -118,6 +170,12 @@ public:
 	}
 	UE_API void SetFalloffDistance(float InFalloffDistance);
 
+	float GetFalloffPlateauRatio() const
+	{
+		return FalloffPlateauRatio;
+	}
+	UE_API void SetFalloffPlateauRatio(float InFalloffPlateauRatio);
+
 	float GetSplinePolygonErrorTolerance() const
 	{
 		return SplinePolygonErrorTolerance;
@@ -176,12 +234,26 @@ private:
 	UPROPERTY(EditDefaultsOnly, Category = "Hidden", meta = (DisplayName = "Spline Component Name"))
 	FString TemplateSplineComponentName;
 
+	/**
+	* Controls which outputs this spline modifier writes.
+	*/
+	UPROPERTY(EditAnywhere, Category = "Coverage", meta = (Bitmask, BitmaskEnum = "/Script/MeshPartitionEditor.ESplineModifierWriteMode"))
+	uint8 WriteMode = static_cast<uint8>(MeshPartition::ESplineModifierWriteMode::Positions | MeshPartition::ESplineModifierWriteMode::Weights);
+
 	/**
 	* Determines how far in projected world distance the spline affects the mesh to either side of it. A higher value results in a slower slope.
 	*/
 	UPROPERTY(EditAnywhere, Category = "Coverage", meta = (UIMin = 5, UIMax = 1500, ClampMin = 0.1))
 	float FalloffDistance = 50.f;
 
+	/**
+	* Fraction of the total falloff distance that is a flat plateau of full influence before the falloff begins.
+	* At 0, the full distance is falloff (triangular profile). At 0.5, half is plateau and half is falloff (trapezoid).
+	* At 1.0, the entire distance is full influence with a hard edge (rectangular profile).
+	*/
+	UPROPERTY(EditAnywhere, Category = "Coverage", meta = (UIMin = 0, UIMax = 1, ClampMin = 0, ClampMax = 1))
+	float FalloffPlateauRatio = 0.0f;
+
 	/**
 	* When true, the scale Y values along the spline are multiplied by FalloffDistance to get the final falloff distance to use.
 	*/
@@ -312,6 +384,12 @@ private:
 	UPROPERTY(EditAnywhere, Category = "Interior", meta = (EditCondition = "bMeshClosedInterior"))
 	bool bShowDebugSurface = false;
 
+	/**
+	* Weight channels written by this modifier, using the spline falloff as the blending alpha.
+	*/
+	UPROPERTY(EditAnywhere, Category = "WeightChannels")
+	TArray<FSplineModifierWeightEntry> WeightChannels;
+
 	UE_API USplineComponent* GetSplineComponent() const;
 	UE_API void UpdateWhenEdgeFalloffChanges(UCurveBase* Curve, EPropertyChangeType::Type);
 	
