Clean up MeshParittionCollisionGeneration settings and remove Experimental from naming

#Synced-CL 52609914

#Synced-CL 52609914

diff --git a/Engine/Plugins/Experimental/MeshPartition/Source/MeshPartitionEditor/Private/MeshPartitionCollisionGeneration.cpp b/Engine/Plugins/Experimental/MeshPartition/Source/MeshPartitionEditor/Private/MeshPartitionCollisionGeneration.cpp
--- a/Engine/Plugins/Experimental/MeshPartition/Source/MeshPartitionEditor/Private/MeshPartitionCollisionGeneration.cpp
+++ b/Engine/Plugins/Experimental/MeshPartition/Source/MeshPartitionEditor/Private/MeshPartitionCollisionGeneration.cpp
@@ -31,7 +31,7 @@ namespace UE::MeshPartitionCollisionGenerationLocals
 	};
 
 	template<typename SimplificationType>
-	void SimplifyDynamicMesh(UE::Geometry::FDynamicMesh3& EditMesh, const UE::MeshPartition::Collision::FExperimentalCollisionSimplificationSettings& SimplifySettings)
+	void SimplifyDynamicMesh(UE::Geometry::FDynamicMesh3& EditMesh, const UE::MeshPartition::Collision::FCollisionSimplificationSettings& SimplifySettings)
 	{
 		using namespace UE::Geometry;
 
@@ -111,15 +111,23 @@ namespace UE::MeshPartitionCollisionGenerationLocals
 
 		// Scale error measure w/ the normal Z-alignment
 		auto CosDegrees = [](float AngleDeg) -> double { return (double)FMath::Cos(FMath::DegreesToRadians(AngleDeg)); };
-		TFunction<double(const FDynamicMesh3&, int VertexA, int VertexB)> CustomScaleF = 
-			[&SimplifySettings, 
-			ScaleRange = FVector2d(CosDegrees(SimplifySettings.ScaleFromMaxNormalAngle), CosDegrees(SimplifySettings.ScaleToMinNormalAngle))]
+		TFunction<double(const FDynamicMesh3&, int VertexA, int VertexB)> CustomScaleF = nullptr;
+		if (SimplifySettings.bScaleAccuracyViaNormal)
+		{
+			FVector3f AlignNormal = (FVector3f)SimplifySettings.ScaleAccuracyNormalDirection;
+			if (!Normalize(AlignNormal))
+			{
+				AlignNormal = FVector3f::UnitZ();
+			}
+			CustomScaleF = [&SimplifySettings, AlignNormal,
+				ScaleRange = FVector2d(CosDegrees(SimplifySettings.ScaleFromMaxNormalAngle), CosDegrees(SimplifySettings.ScaleToMinNormalAngle))]
 			(const FDynamicMesh3& Mesh, int32 VA, int32 VB) -> double
 			{
-				float ZA = Mesh.GetVertexNormal(VA).Z;
-				float ZB = Mesh.GetVertexNormal(VB).Z;
+				float ZA = Mesh.GetVertexNormal(VA) | AlignNormal;
+				float ZB = Mesh.GetVertexNormal(VB) | AlignNormal;
 				return FMath::GetMappedRangeValueClamped<double>(ScaleRange, { 1., (double)SimplifySettings.LocalAccuracyScale }, FMath::Max(ZA, ZB));
 			};
+		}
 
 		if (bUsingGeometricTolerance)
 		{
@@ -542,7 +550,7 @@ namespace UE::MeshPartitionCollisionGenerationLocals
 
 namespace UE::MeshPartition::Collision
 {
-	void FExperimentalCollisionSimplificationSettings::GatherDependencies(MeshPartition::IDependencyInterface& Dependencies) const
+	void FCollisionSimplificationSettings::GatherDependencies(MeshPartition::IDependencyInterface& Dependencies) const
 	{
 		Dependencies += bSimplifyCollision;
 		Dependencies += bCutMeshAlongMaterialBoundaries;
@@ -553,12 +561,14 @@ namespace UE::MeshPartition::Collision
 			Dependencies += SimplifyMethod;
 			Dependencies += EdgeLength;
 			Dependencies += ErrorTolerance;
+			Dependencies += bScaleAccuracyViaNormal;
 			Dependencies += LocalAccuracyScale;
 			Dependencies += ScaleToMinNormalAngle;
 			Dependencies += ScaleFromMaxNormalAngle;
+			Dependencies += ScaleAccuracyNormalDirection;
 			// Version key to allow experimentation w/ the underlying methods, to be bumped with algorithm changes
-			static FGuid ExperimentalSimplificationVersionKey(TEXT("a11b02bb-493b-4607-9cc3-6922fb8a80ad"));
-			Dependencies += ExperimentalSimplificationVersionKey;
+			static FGuid SimplificationVersionKey(TEXT("8db69a3d-d5b7-4beb-a3e1-c0c69cd630f4"));
+			Dependencies += SimplificationVersionKey;
 		}
 	}
 
@@ -570,11 +580,11 @@ namespace UE::MeshPartition::Collision
 		CollisionData.Mesh.Emplace();
 		FTriMeshCollisionData& MeshCollisionData = *CollisionData.Mesh;
 
-		const FExperimentalCollisionSimplificationSettings& SimplifySettings = Settings.SimplificationSettings;
+		const FCollisionSimplificationSettings& SimplifySettings = Settings.SimplificationSettings;
 		const bool bApplyQEMSimplify = SimplifySettings.bSimplifyCollision &&
-			SimplifySettings.SimplifyMethod == UE::MeshPartition::Collision::EExperimentalCollisionSimplificationMethod::QEM;
+			SimplifySettings.SimplifyMethod == UE::MeshPartition::Collision::ECollisionSimplificationMethod::QEM;
 		const bool bApplyClusterSimplify = SimplifySettings.bSimplifyCollision &&
-			SimplifySettings.SimplifyMethod == UE::MeshPartition::Collision::EExperimentalCollisionSimplificationMethod::Cluster &&
+			SimplifySettings.SimplifyMethod == UE::MeshPartition::Collision::ECollisionSimplificationMethod::Cluster &&
 			SimplifySettings.EdgeLength > 0;
 
 		TArray<UPhysicalMaterial*> ReferencePhysicalMaterials = GetReferencePhysicalMaterials(Settings.DefaultPhysicalMaterial, Settings.PhysicalMaterialChannels);
diff --git a/Engine/Plugins/Experimental/MeshPartition/Source/MeshPartitionEditor/Private/MeshPartitionCollisionGeneration.h b/Engine/Plugins/Experimental/MeshPartition/Source/MeshPartitionEditor/Private/MeshPartitionCollisionGeneration.h
--- a/Engine/Plugins/Experimental/MeshPartition/Source/MeshPartitionEditor/Private/MeshPartitionCollisionGeneration.h
+++ b/Engine/Plugins/Experimental/MeshPartition/Source/MeshPartitionEditor/Private/MeshPartitionCollisionGeneration.h
@@ -13,14 +13,14 @@
 namespace UE::MeshPartition::Collision
 {
 	UENUM()
-	enum class EExperimentalCollisionSimplificationMethod
+	enum class ECollisionSimplificationMethod
 	{
 		Cluster,
 		QEM
 	};
 
 	USTRUCT()
-	struct FExperimentalCollisionSimplificationSettings
+	struct FCollisionSimplificationSettings
 	{
 		GENERATED_BODY()
 
@@ -32,39 +32,48 @@ namespace UE::MeshPartition::Collision
 		bool bCutMeshAlongMaterialBoundaries = true;
 
 		// When cutting material boundaries, tolerance distance at which to snap material boundary to existing geometry rather than cutting the mesh
-		UPROPERTY(EditAnywhere, Category = "Material Boundaries", meta = (EditCondition = bCutMeshAlongMaterialBoundaries))
+		UPROPERTY(EditAnywhere, Category = "Material Boundaries", meta = (Units = "cm", EditCondition = bCutMeshAlongMaterialBoundaries))
 		float MaterialBoundaryCutSnapTolerance = 10.f;
 
 		UPROPERTY(EditAnywhere, Category = "Simplification",
-				  meta = (ToolTip = "Enable experimental collision simplification on the separate collision components"))
+				  meta = (ToolTip = "Enable collision simplification on the separate collision components"))
 		bool bSimplifyCollision = false;
 
 		UPROPERTY(EditAnywhere, Category = "Simplification")
-		EExperimentalCollisionSimplificationMethod SimplifyMethod = EExperimentalCollisionSimplificationMethod::QEM;
+		ECollisionSimplificationMethod SimplifyMethod = ECollisionSimplificationMethod::QEM;
 
 		// Constrain how much simplification can move material boundaries away from their original positions
-		UPROPERTY(EditAnywhere, Category = "Simplification", meta = (EditCondition = "SimplifyMethod==EExperimentalCollisionSimplificationMethod::QEM"))
+		UPROPERTY(EditAnywhere, Category = "Simplification", meta = (Units = "cm", EditCondition = "SimplifyMethod==ECollisionSimplificationMethod::QEM"))
 		float MaterialBoundaryDistanceTolerance = 50.f;
 
+		// If positive, set a geometric error tolerance
+		UPROPERTY(EditAnywhere, Category = "Simplification", meta = (Units = "cm", EditCondition = "SimplifyMethod==ECollisionSimplificationMethod::QEM"))
+		float ErrorTolerance = 10.f;
+
 		// If positive, set a target edge length
-		UPROPERTY(EditAnywhere, Category = "Simplification")
-		float EdgeLength = -1;
+		UPROPERTY(EditAnywhere, Category = "Simplification", meta = (Units = "cm"))
+		float EdgeLength = -1.f;
 
-		// If positive, set a geometric error tolerance
-		UPROPERTY(EditAnywhere, Category = "Simplification", meta = (EditCondition = "SimplifyMethod==EExperimentalCollisionSimplificationMethod::QEM"))
-		float ErrorTolerance = 10;
 
-		// Scale up the requested accuracy (error metric and/or measured edge length) by up to this amount for upward-facing surfaces
+		// Whether to scale requested collision accuracy (error tolerance / edge length) locally based on surface normal's alignment with a target direction (e.g. to favor accuracy on 'walkable' surfaces)
 		UPROPERTY(EditAnywhere, Category = "Simplification")
-		float LocalAccuracyScale = 2.5;
+		bool bScaleAccuracyViaNormal = true;
+
+		// Scale up the requested accuracy (error metric and/or measured edge length) by up-to this amount for surfaces w/ normal aligned to ScaleAccuracyNormalDirection
+		UPROPERTY(EditAnywhere, Category = "Simplification|ScaleAccuracy", meta = (EditCondition="bScaleAccuracyViaNormal"))
+		float LocalAccuracyScale = 2.5f;
+
+		// Min normal alignment angle vs ScaleAccuracyNormalDirection, at which to finish scaling up simplifier accuracy, in degrees
+		UPROPERTY(EditAnywhere, Category = "Simplification|ScaleAccuracy", meta = (EditCondition = "bScaleAccuracyViaNormal", ClampMin = 0, ClampMax = 180, Units = Degrees))
+		float ScaleToMinNormalAngle = 26.f;
 
-		// Min normal alignment angle vs Z-up, at which to finish scaling up simplifier accuracy, in degrees
-		UPROPERTY(EditAnywhere, Category = "Simplification", meta = (ClampMin = 0, ClampMax = 180, Units = Degrees))
-		float ScaleToMinNormalAngle = 26;
+		// Max normal alignment angle vs ScaleAccuracyNormalDirection, at which to start scaling up simplifier accuracy, in degrees
+		UPROPERTY(EditAnywhere, Category = "Simplification|ScaleAccuracy", meta = (EditCondition = "bScaleAccuracyViaNormal", ClampMin=0, ClampMax=180, Units = Degrees))
+		float ScaleFromMaxNormalAngle = 53.f;
 
-		// Max normal alignment angle vs Z-up, at which to start scaling up simplifier accuracy, in degrees
-		UPROPERTY(EditAnywhere, Category = "Simplification", meta = (ClampMin=0, ClampMax=180, Units = Degrees))
-		float ScaleFromMaxNormalAngle = 53;
+		// The normal direction that should have highest accuracy
+		UPROPERTY(EditAnywhere, Category = "Simplification|ScaleAccuracy", meta = (EditCondition = "bScaleAccuracyViaNormal"))
+		FVector3d ScaleAccuracyNormalDirection = FVector(0., 0., 1.);
 
 		#if WITH_EDITOR
 		void MESHPARTITIONEDITOR_API GatherDependencies(MeshPartition::IDependencyInterface& Dependencies) const;
@@ -73,7 +82,7 @@ namespace UE::MeshPartition::Collision
 
 	struct FMeshToCollisionSettings
 	{
-		UE::MeshPartition::Collision::FExperimentalCollisionSimplificationSettings SimplificationSettings;
+		UE::MeshPartition::Collision::FCollisionSimplificationSettings SimplificationSettings;
 		
 		/** Physical material channels; vertex weight channels will be translated to per - triangle physical materials via this mapping */
 		TArray<UE::MeshPartition::FPhysicalMaterialChannel> PhysicalMaterialChannels;
diff --git a/Engine/Plugins/Experimental/MeshPartition/Source/MeshPartitionEditor/Private/MeshPartitionCollisionTransformer.cpp b/Engine/Plugins/Experimental/MeshPartition/Source/MeshPartitionEditor/Private/MeshPartitionCollisionTransformer.cpp
--- a/Engine/Plugins/Experimental/MeshPartition/Source/MeshPartitionEditor/Private/MeshPartitionCollisionTransformer.cpp
+++ b/Engine/Plugins/Experimental/MeshPartition/Source/MeshPartitionEditor/Private/MeshPartitionCollisionTransformer.cpp
@@ -72,7 +72,7 @@ bool FCollisionTransformer::Execute(MeshPartition::FTransformerContext& InTransf
 
 void FCollisionTransformer::GatherDependencies(MeshPartition::IDependencyInterface& InDependencies) const
 {
-	ExperimentalCollisionSimplificationSettings.GatherDependencies(InDependencies);
+	CollisionSimplificationSettings.GatherDependencies(InDependencies);
 	InDependencies += CollisionProfile.Name;
 	InDependencies += bCanEverAffectNavigation;
 	InDependencies += bFastCook;
@@ -97,7 +97,7 @@ void FCollisionTransformer::BuildCollisionData(const MeshPartition::FTransformer
 	    		
 	Collision::FMeshToCollisionSettings Settings;
 			
-	Settings.SimplificationSettings = ExperimentalCollisionSimplificationSettings;
+	Settings.SimplificationSettings = CollisionSimplificationSettings;
 	Settings.bFastCook = bFastCook; //todo(luc.eygasier): should we take into account preview vs compiled sections? (true also for disable active edge precompute)
 	Settings.bDisableActiveEdgePrecompute = bDisableActiveEdgePrecompute;
 	Settings.PhysicalMaterialChannels = PhysicalMaterialChannels;
diff --git a/Engine/Plugins/Experimental/MeshPartition/Source/MeshPartitionEditor/Public/MeshPartitionCollisionTransformer.h b/Engine/Plugins/Experimental/MeshPartition/Source/MeshPartitionEditor/Public/MeshPartitionCollisionTransformer.h
--- a/Engine/Plugins/Experimental/MeshPartition/Source/MeshPartitionEditor/Public/MeshPartitionCollisionTransformer.h
+++ b/Engine/Plugins/Experimental/MeshPartition/Source/MeshPartitionEditor/Public/MeshPartitionCollisionTransformer.h
@@ -44,9 +44,9 @@ private:
 	UPROPERTY()
 	TObjectPtr<UPhysicalMaterial> DefaultPhysicalMaterial;
 
-	// Experimental collision simplification settings, only applied to the separate collision components. Only for testing/evaluation purposes.
+	// Collision simplification settings, only applied to the separate collision components.
 	UPROPERTY(EditAnywhere, AdvancedDisplay, Category = "Simplification", meta=(EditCondition = "bEnableSeparateCollisionComponents"))
-	MeshPartition::Collision::FExperimentalCollisionSimplificationSettings ExperimentalCollisionSimplificationSettings;
+	MeshPartition::Collision::FCollisionSimplificationSettings CollisionSimplificationSettings;
 
 	UPROPERTY(EditAnywhere, Category = "Collision")
 	FCollisionProfileName CollisionProfile = UCollisionProfile::BlockAll_ProfileName;
