[Mesh Partition] Only show custom context menu when Ctrl+right-click

Previously, right-clicking any MeshPartition preview section would intercept the click and show the modifier priority stack custom selection menu, preventing the standard UE actor context menu from appearing on a plain right-click.

#Synced-CL 52678297

diff --git a/Engine/Plugins/Experimental/MeshPartition/Source/MeshPartitionEditor/Private/MeshPartitionPreviewVisualizers.cpp b/Engine/Plugins/Experimental/MeshPartition/Source/MeshPartitionEditor/Private/MeshPartitionPreviewVisualizers.cpp
--- a/Engine/Plugins/Experimental/MeshPartition/Source/MeshPartitionEditor/Private/MeshPartitionPreviewVisualizers.cpp
+++ b/Engine/Plugins/Experimental/MeshPartition/Source/MeshPartitionEditor/Private/MeshPartitionPreviewVisualizers.cpp
@@ -337,6 +337,11 @@ namespace MegaMeshPreviewVisualizersLocals
 		
 		if (Click.GetKey() == EKeys::RightMouseButton)
 		{
+			// Only summon the custom right click menu if Ctrl is pressed.
+			if (!Click.IsControlDown())
+			{
+				return false;
+			}
 			return SummonRightClickMenu(Visualizer, ViewportClient, Click, 
 				PreviewSectionComponent,
 				*ElementSelectionSet, 
