Temporarily disable separate world builder path creating new compiled sections.

#Synced-CL 52707165
#Synced-CL 52707165

diff --git a/Engine/Plugins/Experimental/MeshPartition/Source/MeshPartitionEditor/Private/MeshPartitionEditorSubsystem.cpp b/Engine/Plugins/Experimental/MeshPartition/Source/MeshPartitionEditor/Private/MeshPartitionEditorSubsystem.cpp
--- a/Engine/Plugins/Experimental/MeshPartition/Source/MeshPartitionEditor/Private/MeshPartitionEditorSubsystem.cpp
+++ b/Engine/Plugins/Experimental/MeshPartition/Source/MeshPartitionEditor/Private/MeshPartitionEditorSubsystem.cpp
@@ -24,7 +24,7 @@ extern UNREALED_API UEditorEngine* GEditor;
 namespace UE::MeshPartition
 {
 TAutoConsoleVariable<bool> CVarMeshPartitionPIECreateMissing(TEXT("MeshPartition.PIE.CreateMissingSections"),
-	true,
+	false,
 	TEXT("Whether to create any missing (or stripped) compiled sections when playing in editor."));
 
 TAutoConsoleVariable<int32> CVarGCThresholdMB(TEXT("MeshPartition.Preview.GCThresholdMB"),
