Mesh Partition: remove the unused Tessellate modifier. Redirect existing Tessellate modifiers to RemeshModifier

#Synced-CL 52741992

diff --git a/Engine/Plugins/Experimental/MeshPartition/Config/DefaultMeshPartition.ini b/Engine/Plugins/Experimental/MeshPartition/Config/DefaultMeshPartition.ini
--- a/Engine/Plugins/Experimental/MeshPartition/Config/DefaultMeshPartition.ini
+++ b/Engine/Plugins/Experimental/MeshPartition/Config/DefaultMeshPartition.ini
@@ -1,2 +1,11 @@
 [/Script/MeshPartition.Settings]
-DefaultDefinition=/MeshPartition/DataAssets/MPD_Default.MPD_Default
\ No newline at end of file
+DefaultDefinition=/MeshPartition/DataAssets/MPD_Default.MPD_Default
+
+[CoreRedirects]
+; Tessellate modifier removed, redirected to Remesh modifier
++EnumRedirects=(OldName="/Script/MegaMeshEditor.MegaMeshTessellateModifierMethod", NewName="/Script/MeshPartitionEditor.MegaMeshRemeshModifierTessellateMethod")
++ClassRedirects=(OldName="/Script/MegaMeshEditor.MegaMeshTessellateModifier", NewName="/Script/MeshPartitionEditor.RemeshModifier")
++EnumRedirects=(OldName="/Script/MegaMeshEditor.TessellateModifierMethod", NewName="/Script/MeshPartitionEditor.MegaMeshRemeshModifierTessellateMethod")
++ClassRedirects=(OldName="/Script/MegaMeshEditor.TessellateModifier", NewName="/Script/MeshPartitionEditor.RemeshModifier")
++EnumRedirects=(OldName="/Script/MeshPartitionEditor.TessellateModifierMethod", NewName="/Script/MeshPartitionEditor.MegaMeshRemeshModifierTessellateMethod")
++ClassRedirects=(OldName="/Script/MeshPartitionEditor.TessellateModifier", NewName="/Script/MeshPartitionEditor.RemeshModifier")
\ No newline at end of file
diff --git a/Engine/Plugins/Experimental/MeshPartition/Source/MeshPartitionEditor/Private/Modifiers/MeshPartitionTessellateModifier.cpp b/Engine/Plugins/Experimental/MeshPartition/Source/MeshPartitionEditor/Private/Modifiers/MeshPartitionTessellateModifier.cpp
deleted file mode 100644
diff --git a/Engine/Plugins/Experimental/MeshPartition/Source/MeshPartitionEditor/Public/Modifiers/MeshPartitionTessellateModifier.h b/Engine/Plugins/Experimental/MeshPartition/Source/MeshPartitionEditor/Public/Modifiers/MeshPartitionTessellateModifier.h
deleted file mode 100644
diff --git a/Engine/Plugins/Experimental/MeshPartition/Source/MeshPartitionEditorUI/Private/UI/MeshPartitionLayerNameWidget.cpp b/Engine/Plugins/Experimental/MeshPartition/Source/MeshPartitionEditorUI/Private/UI/MeshPartitionLayerNameWidget.cpp
--- a/Engine/Plugins/Experimental/MeshPartition/Source/MeshPartitionEditorUI/Private/UI/MeshPartitionLayerNameWidget.cpp
+++ b/Engine/Plugins/Experimental/MeshPartition/Source/MeshPartitionEditorUI/Private/UI/MeshPartitionLayerNameWidget.cpp
@@ -48,7 +48,6 @@ namespace MegaMeshNameWidgetLocals
 		{MeshPartition::UProjectMeshLayersModifier::StaticClass(), TEXT("ModifierSculpt")},
 		{MeshPartition::URemeshModifier::StaticClass(), TEXT("ModifierRemesh")},
 		{MeshPartition::USplineModifier::StaticClass(), TEXT("ModifierPath")},
-		//{MeshPartition::UTessellateModifier::StaticClass(), TEXT("Icons.Star")},
 		{MeshPartition::UTexturePatchModifier::StaticClass(), TEXT("ModifierTexturePatch")},
 	};
 
diff --git a/Engine/Plugins/Experimental/MeshPartition/Source/MeshPartitionModelingToolset/Private/MeshPartitionPlaceModifierTool.cpp b/Engine/Plugins/Experimental/MeshPartition/Source/MeshPartitionModelingToolset/Private/MeshPartitionPlaceModifierTool.cpp
--- a/Engine/Plugins/Experimental/MeshPartition/Source/MeshPartitionModelingToolset/Private/MeshPartitionPlaceModifierTool.cpp
+++ b/Engine/Plugins/Experimental/MeshPartition/Source/MeshPartitionModelingToolset/Private/MeshPartitionPlaceModifierTool.cpp
@@ -27,7 +27,6 @@
 #include "Modifiers/MeshPartitionNoiseModifier.h"
 #include "Modifiers/MeshPartitionBooleanModifier.h"
 #include "Modifiers/MeshPartitionLatticeModifier.h"
-#include "Modifiers/MeshPartitionTessellateModifier.h"
 #include "Modifiers/MeshPartitionSplineRemeshModifier.h"
 #include "Modifiers/MeshPartitionWeightUtilityModifier.h"
 
@@ -237,8 +236,6 @@ TObjectPtr<UClass> UPlaceModifierTool::GetModifierToCreateClass()
 		return MeshPartition::UBooleanModifier::StaticClass();
 	case MeshPartition::EModifierClassType::Lattice:
 		return MeshPartition::ULatticeModifier::StaticClass();
-	case MeshPartition::EModifierClassType::Tessellate:
-		return MeshPartition::UTessellateModifier::StaticClass();
 	case MeshPartition::EModifierClassType::SplineRemesh:
 		return MeshPartition::USplineRemeshModifier::StaticClass();
 	case MeshPartition::EModifierClassType::WeightUtility:
@@ -315,8 +312,6 @@ void UPlaceModifierTool::SetupNewModifier(AActor* ContainingActor, MeshPartition
 		break;
 	case MeshPartition::EModifierClassType::Lattice:
 		break;
-	case MeshPartition::EModifierClassType::Tessellate:
-		break;
 	case MeshPartition::EModifierClassType::SplineRemesh:
 	{
 		FCreateComponentParams Params;
diff --git a/Engine/Plugins/Experimental/MeshPartition/Source/MeshPartitionModelingToolset/Public/MeshPartitionPlaceModifierTool.h b/Engine/Plugins/Experimental/MeshPartition/Source/MeshPartitionModelingToolset/Public/MeshPartitionPlaceModifierTool.h
--- a/Engine/Plugins/Experimental/MeshPartition/Source/MeshPartitionModelingToolset/Public/MeshPartitionPlaceModifierTool.h
+++ b/Engine/Plugins/Experimental/MeshPartition/Source/MeshPartitionModelingToolset/Public/MeshPartitionPlaceModifierTool.h
@@ -28,7 +28,6 @@ enum class EModifierClassType : uint8
 	Noise,
 	Boolean,
 	Lattice,
-	Tessellate,
 	SplineRemesh,
 	WeightUtility,
 
diff --git a/Engine/Plugins/Experimental/MeshTerrainMode/Source/MeshTerrainMode/Private/MeshTerrainMode.cpp b/Engine/Plugins/Experimental/MeshTerrainMode/Source/MeshTerrainMode/Private/MeshTerrainMode.cpp
--- a/Engine/Plugins/Experimental/MeshTerrainMode/Source/MeshTerrainMode/Private/MeshTerrainMode.cpp
+++ b/Engine/Plugins/Experimental/MeshTerrainMode/Source/MeshTerrainMode/Private/MeshTerrainMode.cpp
@@ -1052,7 +1052,6 @@ void UMeshTerrainMode::Enter()
 	RegisterPlaceModifierTool(ToolManagerCommands.BeginAddNoiseModifierTool, TEXT("BeginAddNoiseModifierTool"), static_cast<int>(UE::MeshPartition::EModifierClassType::Noise));
 	RegisterPlaceModifierTool(ToolManagerCommands.BeginAddBooleanModifierTool, TEXT("BeginAddBooleanModifierTool"), static_cast<int>(UE::MeshPartition::EModifierClassType::Boolean));
 	RegisterPlaceModifierTool(ToolManagerCommands.BeginAddLatticeModifierTool, TEXT("BeginAddLatticeModifierTool"), static_cast<int>(UE::MeshPartition::EModifierClassType::Lattice));
-	RegisterPlaceModifierTool(ToolManagerCommands.BeginAddTessellateModifierTool, TEXT("BeginAddTessellateModifierTool"), static_cast<int>(UE::MeshPartition::EModifierClassType::Tessellate));
 	RegisterPlaceModifierTool(ToolManagerCommands.BeginAddSplineRemeshModifierTool, TEXT("BeginAddSplineRemeshModifierTool"), static_cast<int>(UE::MeshPartition::EModifierClassType::SplineRemesh));
 
 	// register extensions
diff --git a/Engine/Plugins/Experimental/MeshTerrainMode/Source/MeshTerrainMode/Private/MeshTerrainModeManagerActions.cpp b/Engine/Plugins/Experimental/MeshTerrainMode/Source/MeshTerrainMode/Private/MeshTerrainModeManagerActions.cpp
--- a/Engine/Plugins/Experimental/MeshTerrainMode/Source/MeshTerrainMode/Private/MeshTerrainModeManagerActions.cpp
+++ b/Engine/Plugins/Experimental/MeshTerrainMode/Source/MeshTerrainMode/Private/MeshTerrainModeManagerActions.cpp
@@ -328,7 +328,6 @@ void FMeshTerrainModeManagerCommands::RegisterCommands()
 	UI_COMMAND(BeginAddNoiseModifierTool, "Noise", "Add Noise Modifier", EUserInterfaceActionType::ToggleButton, FInputChord());
 	UI_COMMAND(BeginAddBooleanModifierTool, "Boolean", "Add Boolean Modifier", EUserInterfaceActionType::ToggleButton, FInputChord());
 	UI_COMMAND(BeginAddLatticeModifierTool, "Lattice", "Add Lattice Modifier", EUserInterfaceActionType::ToggleButton, FInputChord());
-	UI_COMMAND(BeginAddTessellateModifierTool, "Tessellate", "Add Tessellate Modifier", EUserInterfaceActionType::ToggleButton, FInputChord());
 	UI_COMMAND(BeginAddSplineRemeshModifierTool, "Spline Remesh", "Add Spline Remesh Modifier", EUserInterfaceActionType::ToggleButton, FInputChord());
 
 	UI_COMMAND(EnterShapesSubmode, "Shapes", "Shapes", EUserInterfaceActionType::ToggleButton, FInputChord());
diff --git a/Engine/Plugins/Experimental/MeshTerrainMode/Source/MeshTerrainMode/Private/MeshTerrainModeManagerActions.h b/Engine/Plugins/Experimental/MeshTerrainMode/Source/MeshTerrainMode/Private/MeshTerrainModeManagerActions.h
--- a/Engine/Plugins/Experimental/MeshTerrainMode/Source/MeshTerrainMode/Private/MeshTerrainModeManagerActions.h
+++ b/Engine/Plugins/Experimental/MeshTerrainMode/Source/MeshTerrainMode/Private/MeshTerrainModeManagerActions.h
@@ -270,7 +270,6 @@ public:
 	TSharedPtr<FUICommandInfo> BeginAddNoiseModifierTool;
 	TSharedPtr<FUICommandInfo> BeginAddBooleanModifierTool;
 	TSharedPtr<FUICommandInfo> BeginAddLatticeModifierTool;
-	TSharedPtr<FUICommandInfo> BeginAddTessellateModifierTool;
 	TSharedPtr<FUICommandInfo> BeginAddSplineRemeshModifierTool;
 
 	// Submodes
diff --git a/Engine/Plugins/Experimental/MeshTerrainMode/Source/MeshTerrainMode/Private/MeshTerrainModeStyle.cpp b/Engine/Plugins/Experimental/MeshTerrainMode/Source/MeshTerrainMode/Private/MeshTerrainModeStyle.cpp
--- a/Engine/Plugins/Experimental/MeshTerrainMode/Source/MeshTerrainMode/Private/MeshTerrainModeStyle.cpp
+++ b/Engine/Plugins/Experimental/MeshTerrainMode/Source/MeshTerrainMode/Private/MeshTerrainModeStyle.cpp
@@ -527,8 +527,6 @@ void FMeshTerrainModeStyle::Initialize()
 		StyleSet->Set("MeshTerrainModeManagerCommands.BeginAddBooleanModifierTool.Small",			new IMAGE_BRUSH_SVG("Icons/ModifierBoolean_16", Icon20x20));
 		StyleSet->Set("MeshTerrainModeManagerCommands.BeginAddLatticeModifierTool",				new IMAGE_BRUSH_SVG("Icons/ModifierLattice_16", Icon40x40));
 		StyleSet->Set("MeshTerrainModeManagerCommands.BeginAddLatticeModifierTool.Small",			new IMAGE_BRUSH_SVG("Icons/ModifierLattice_16", Icon20x20));
-		StyleSet->Set("MeshTerrainModeManagerCommands.BeginAddTessellateModifierTool",				new IMAGE_BRUSH_SVG("Icons/ModifierTessellation_16", Icon40x40));
-		StyleSet->Set("MeshTerrainModeManagerCommands.BeginAddTessellateModifierTool.Small",		new IMAGE_BRUSH_SVG("Icons/ModifierTessellation_16", Icon20x20));
 		StyleSet->Set("MeshTerrainModeManagerCommands.BeginAddSplineRemeshModifierTool",			new IMAGE_BRUSH_SVG("Icons/ModifierPath_16", Icon40x40));
 		StyleSet->Set("MeshTerrainModeManagerCommands.BeginAddSplineRemeshModifierTool.Small",		new IMAGE_BRUSH_SVG("Icons/ModifierPath_16", Icon20x20));
 
