[Backout] - CL52747137
#rnx
Original CL Desc
-----------------------------------------------------------------
Mesh Partition: remove the unused Tessellate modifier. Redirect existing Tessellate modifiers to RemeshModifier

diff --git a/Engine/Plugins/Experimental/MeshPartition/Config/DefaultMeshPartition.ini b/Engine/Plugins/Experimental/MeshPartition/Config/DefaultMeshPartition.ini
--- a/Engine/Plugins/Experimental/MeshPartition/Config/DefaultMeshPartition.ini
+++ b/Engine/Plugins/Experimental/MeshPartition/Config/DefaultMeshPartition.ini
@@ -1,11 +1,2 @@
 [/Script/MeshPartition.Settings]
-DefaultDefinition=/MeshPartition/DataAssets/MPD_Default.MPD_Default
-
-[CoreRedirects]
-; Tessellate modifier removed, redirected to Remesh modifier
-+EnumRedirects=(OldName="/Script/MegaMeshEditor.MegaMeshTessellateModifierMethod", NewName="/Script/MeshPartitionEditor.MegaMeshRemeshModifierTessellateMethod")
-+ClassRedirects=(OldName="/Script/MegaMeshEditor.MegaMeshTessellateModifier", NewName="/Script/MeshPartitionEditor.RemeshModifier")
-+EnumRedirects=(OldName="/Script/MegaMeshEditor.TessellateModifierMethod", NewName="/Script/MeshPartitionEditor.MegaMeshRemeshModifierTessellateMethod")
-+ClassRedirects=(OldName="/Script/MegaMeshEditor.TessellateModifier", NewName="/Script/MeshPartitionEditor.RemeshModifier")
-+EnumRedirects=(OldName="/Script/MeshPartitionEditor.TessellateModifierMethod", NewName="/Script/MeshPartitionEditor.MegaMeshRemeshModifierTessellateMethod")
-+ClassRedirects=(OldName="/Script/MeshPartitionEditor.TessellateModifier", NewName="/Script/MeshPartitionEditor.RemeshModifier")
\ No newline at end of file
+DefaultDefinition=/MeshPartition/DataAssets/MPD_Default.MPD_Default
\ No newline at end of file
diff --git a/Engine/Plugins/Experimental/MeshPartition/Source/MeshPartitionEditor/Private/Modifiers/MeshPartitionTessellateModifier.cpp b/Engine/Plugins/Experimental/MeshPartition/Source/MeshPartitionEditor/Private/Modifiers/MeshPartitionTessellateModifier.cpp
new file mode 100644
--- /dev/null
+++ b/Engine/Plugins/Experimental/MeshPartition/Source/MeshPartitionEditor/Private/Modifiers/MeshPartitionTessellateModifier.cpp
@@ -0,0 +1,264 @@
+// Copyright Epic Games, Inc. All Rights Reserved.
+
+#include "Modifiers/MeshPartitionTessellateModifier.h"
+
+#include "MeshPartitionModule.h" // FCustomVersion
+#include "Operations/SelectiveTessellate.h"
+#include "PrimitiveDrawingUtils.h"
+#include "Selections/GeometrySelectionUtil.h"
+
+namespace UE::MeshPartition
+{
+namespace MegaMeshTessellateModifierLocals
+{
+	class FTessellateBackgroundOp : public MeshPartition::IModifierBackgroundOp
+	{
+	public:
+		FTessellateBackgroundOp(const FName& InOperationName) : MeshPartition::IModifierBackgroundOp(InOperationName) {}
+
+		FBox LocalBounds;
+		FBox GlobalBounds;
+
+		bool bUseTargetEdgeLength;
+		int32 TessellationLevel;
+		float TargetEdgeLength;
+		int32 MaxTessellationLevel;
+
+		MeshPartition::ETessellateModifierMethod TessellationMethod;
+		FTransform ModifierTransform;
+
+		// Generate a new random guid before submitting any code changes to the op
+		static FGuid GetCodeVersionKey()
+		{
+			static FGuid VersionKey(TEXT("1812dec5-557f-40df-bce7-4c8596e75b9e"));
+			return VersionKey;
+		}
+
+	private:
+
+		virtual void GetInstancesInBounds(const FBox& InBounds, TArray<FInstanceInfo>& OutInstanceInfos) const override;
+		virtual void ApplyModifications(MeshPartition::FMeshView& InMeshView, const FTransform3d& InTransform, const FInstanceInfo& InInstanceDesc) const override;
+
+		// Set to true whenever iterating on code changes to prevent any builds including this modifier being picked up by ddc
+		// and poisoning the cache/generating lots of unused intermediate data.
+		virtual bool DisableDDCWrite() const override 
+		{ 
+			return false; 
+		}
+	};
+
+}
+
+UTessellateModifier::UTessellateModifier()
+{
+}
+
+void UTessellateModifier::Serialize(FArchive& Ar)
+{
+	Super::Serialize(Ar);
+
+	const FName ParentPreviousDefaultType = TEXT("Misc");
+	if (Ar.IsLoading()
+		&& GetLinkerCustomVersion(UE::MeshPartition::FCustomVersion::GUID) < UE::MeshPartition::FCustomVersion::DefaultPriorityLayerSetToNone
+		&& GetType() == ParentPreviousDefaultType) // Super::Serialize would set default to this
+	{
+		const FName PreviousModifierDefaultType = TEXT("Tessellate");
+		SetType(PreviousModifierDefaultType);
+	}
+}
+
+FBox UTessellateModifier::ComputeLocalBounds() const
+{
+	return FBox(-0.5* Extents, 0.5*Extents);
+}
+
+TArray<FBox> UTessellateModifier::ComputeBounds() const
+{	
+	return { ComputeLocalBounds().TransformBy(GetComponentToWorld()) };
+}
+
+void MegaMeshTessellateModifierLocals::FTessellateBackgroundOp::GetInstancesInBounds(const FBox& InBounds, TArray<FInstanceInfo>& OutInstanceInfos) const
+{
+	AddDefaultInstanceIfIntersects(GlobalBounds, InBounds, OutInstanceInfos);
+
+	if (!OutInstanceInfos.IsEmpty())
+	{
+		OutInstanceInfos[0].ReadViewComponents = EMeshViewComponents::DynamicSubmesh;
+		OutInstanceInfos[0].WriteViewComponents = EMeshViewComponents::DynamicSubmesh;
+	}
+}
+
+TSharedPtr<const MeshPartition::IModifierBackgroundOp> UTessellateModifier::CreateBackgroundOp(const MeshPartition::EBuildType InBuildType) const
+{
+	using namespace MegaMeshTessellateModifierLocals;
+	TSharedPtr<FTessellateBackgroundOp> Op = MakeShared<FTessellateBackgroundOp>(GetFName());
+	Op->GlobalBounds = ComputeCombinedBounds();
+	Op->LocalBounds = ComputeLocalBounds();
+	
+	Op->bUseTargetEdgeLength = bUseTargetEdgeLength;
+	Op->TessellationLevel = Level;
+	Op->TargetEdgeLength = TargetEdgeLength;
+	Op->MaxTessellationLevel = MaxLevel;
+
+	Op->TessellationMethod = TessellationMethod;
+	Op->ModifierTransform = GetComponentToWorld();
+	return Op;
+}
+
+FGuid UTessellateModifier::GetCodeVersionKey() const
+{
+	return MegaMeshTessellateModifierLocals::FTessellateBackgroundOp::GetCodeVersionKey();
+}
+
+void MegaMeshTessellateModifierLocals::FTessellateBackgroundOp::ApplyModifications(MeshPartition::FMeshView& InMeshView, const FTransform3d& InTransform, const FInstanceInfo& InInstanceDesc) const
+{
+	TRACE_CPUPROFILER_EVENT_SCOPE(MeshPartition::UTessellateModifier::ApplyModifications);
+
+	using namespace Geometry;
+
+	FDynamicMesh3& SubMesh = InMeshView.GetSubmeshMutable();
+
+	TSet<int32> InsideTriangles;
+	for (int32 VertexID = 0; VertexID < SubMesh.VertexCount(); ++VertexID)
+	{
+		const FVector3d MeshVertexPosition = SubMesh.GetVertex(VertexID);
+		const FVector3d ModifierSpaceVertexPosition = ModifierTransform.InverseTransformPosition(MeshVertexPosition);
+
+		if (LocalBounds.IsInside(ModifierSpaceVertexPosition))
+		{
+			for (const int32 TriIndex : SubMesh.VtxTrianglesItr(VertexID))
+			{
+				InsideTriangles.Add(TriIndex);
+			}
+		}
+	}
+
+	FDynamicMesh3 InputMesh;
+	InputMesh.Copy(SubMesh);
+
+	double ActualTargetEdgeLength = TargetEdgeLength;
+	if (!ensure(TargetEdgeLength > 0.0))		// should be enforced by UPROPERTY meta ClampMin
+	{
+		ActualTargetEdgeLength = UE_KINDA_SMALL_NUMBER;
+	}
+
+	TFunction<int32(int32)> TessellationForTriangleFunc;
+	if (bUseTargetEdgeLength)
+	{
+		TessellationForTriangleFunc = [&InputMesh, ActualTargetEdgeLength, this](int32 TriangleIndex) -> int32
+		{
+			auto TessellationForEdgeFunc = [&InputMesh, ActualTargetEdgeLength, this](int32 EdgeIndex) -> int32
+			{
+				if (InputMesh.IsEdge(EdgeIndex))
+				{
+					FVector3d EdgeEndA, EdgeEndB;
+					InputMesh.GetEdgeV(EdgeIndex, EdgeEndA, EdgeEndB);
+					const double CurrEdgeLength = FVector3d::Dist(EdgeEndA, EdgeEndB);
+					const int32 EdgeLevel = FMath::Clamp(FMath::RoundToInt32(CurrEdgeLength / ActualTargetEdgeLength) - 1, 0, MaxTessellationLevel);
+					return EdgeLevel;
+				}
+				return 0;
+			};
+
+			if (InputMesh.IsTriangle(TriangleIndex))
+			{
+				const FIndex3i Edges = InputMesh.GetTriEdges(TriangleIndex);
+				return FMath::Max3(TessellationForEdgeFunc(Edges[0]), TessellationForEdgeFunc(Edges[1]), TessellationForEdgeFunc(Edges[2]));
+			}
+			return 0;
+		};
+	}
+	else
+	{
+		TessellationForTriangleFunc = [this](int32)
+		{
+			return TessellationLevel;
+		};
+	}
+
+
+	switch (TessellationMethod)
+	{
+	case MeshPartition::ETessellateModifierMethod::UniformRings:
+	{
+		TArray<int> TriTessLevels;
+		TriTessLevels.Init(0, InputMesh.MaxTriangleID());
+		for (const int32 InsideTriangleIndex : InsideTriangles)
+		{
+			TriTessLevels[InsideTriangleIndex] = TessellationForTriangleFunc(InsideTriangleIndex);
+		}
+
+		TArray<int> EdgeTessLevels;
+		EdgeTessLevels.Init(0, InputMesh.MaxEdgeID());
+		for (int32 TriID : InputMesh.TriangleIndicesItr())
+		{
+			const FIndex3i Edges = InputMesh.GetTriEdges(TriID);
+			EdgeTessLevels[Edges[0]] = FMath::Min(EdgeTessLevels[Edges[0]], TriTessLevels[TriID]);
+			EdgeTessLevels[Edges[1]] = FMath::Min(EdgeTessLevels[Edges[1]], TriTessLevels[TriID]);
+			EdgeTessLevels[Edges[2]] = FMath::Min(EdgeTessLevels[Edges[2]], TriTessLevels[TriID]);
+		}
+		FSelectiveTessellate Tessellate(&InputMesh, &SubMesh);
+		TUniquePtr<FTessellationPattern> Pattern = FSelectiveTessellate::CreateConcentricRingsTessellationPattern(&InputMesh, EdgeTessLevels, TriTessLevels);
+		Tessellate.Pattern = Pattern.Get();
+		Tessellate.Compute();
+	}
+	break;
+	case MeshPartition::ETessellateModifierMethod::AdaptiveRings:
+	{
+		TArray<int> TriTessLevels;
+		TaperPerTriangleValues(InputMesh, InsideTriangles, TessellationForTriangleFunc, TriTessLevels);
+
+		TArray<int> EdgeTessLevels;
+		EdgeTessLevels.Init(0, InputMesh.MaxEdgeID());
+		for (int32 TriID : InputMesh.TriangleIndicesItr())
+		{
+			const FIndex3i Edges = InputMesh.GetTriEdges(TriID);
+			EdgeTessLevels[Edges[0]] = FMath::Min(EdgeTessLevels[Edges[0]], TriTessLevels[TriID]);
+			EdgeTessLevels[Edges[1]] = FMath::Min(EdgeTessLevels[Edges[1]], TriTessLevels[TriID]);
+			EdgeTessLevels[Edges[2]] = FMath::Min(EdgeTessLevels[Edges[2]], TriTessLevels[TriID]);
+		}
+
+		FSelectiveTessellate Tessellate(&InputMesh, &SubMesh);
+		TUniquePtr<FTessellationPattern> Pattern = FSelectiveTessellate::CreateConcentricRingsTessellationPattern(&InputMesh, EdgeTessLevels, TriTessLevels);
+		Tessellate.Pattern = Pattern.Get();
+		Tessellate.Compute();
+	}
+	break;
+	case MeshPartition::ETessellateModifierMethod::AdaptiveRegular:
+	{
+		TArray<int> TriTessLevels;
+		TaperPerTriangleValues(SubMesh, InsideTriangles, TessellationForTriangleFunc, TriTessLevels);
+
+		FSelectiveTessellate Tessellate(&InputMesh, &SubMesh);
+		TUniquePtr<FTessellationPattern> Pattern = FSelectiveTessellate::CreateRedGreenTessellationPattern(&InputMesh, TriTessLevels);
+		Tessellate.Pattern = Pattern.Get();
+		Tessellate.Compute();
+	}
+	break;
+	};
+}
+
+void UTessellateModifier::GatherDependencies(MeshPartition::IDependencyInterface& Dependencies) const
+{
+	TRACE_CPUPROFILER_EVENT_SCOPE(UTessellateModifier::GatherDependencies);
+
+	Super::GatherDependencies(Dependencies);
+
+	Dependencies += ComputeCombinedBounds();
+	Dependencies += Extents;
+	Dependencies += TessellationMethod;
+	Dependencies += bUseTargetEdgeLength;
+	Dependencies += Level;
+	Dependencies += TargetEdgeLength;
+	Dependencies += MaxLevel;
+}
+
+void UTessellateModifier::DrawVisualization(const FSceneView* View, FPrimitiveDrawInterface* PDI) const
+{
+	if (bDrawLocalBounds)
+	{
+		const FMatrix Matrix = GetComponentToWorld().ToMatrixWithScale();
+		DrawWireBox(PDI, Matrix, ComputeLocalBounds(), FLinearColor::Yellow, SDPG_World, 1.0, 0, true);
+	}
+}
+}
\ No newline at end of file
diff --git a/Engine/Plugins/Experimental/MeshPartition/Source/MeshPartitionEditor/Public/Modifiers/MeshPartitionTessellateModifier.h b/Engine/Plugins/Experimental/MeshPartition/Source/MeshPartitionEditor/Public/Modifiers/MeshPartitionTessellateModifier.h
new file mode 100644
--- /dev/null
+++ b/Engine/Plugins/Experimental/MeshPartition/Source/MeshPartitionEditor/Public/Modifiers/MeshPartitionTessellateModifier.h
@@ -0,0 +1,80 @@
+// Copyright Epic Games, Inc. All Rights Reserved.
+
+#pragma once
+
+#include "MeshPartitionModifierComponent.h"
+#include "MeshPartitionTessellateModifier.generated.h"
+
+#define UE_API MESHPARTITIONEDITOR_API
+
+namespace UE::MeshPartition
+{
+
+UENUM()
+enum class ETessellateModifierMethod : uint8
+{
+	/** GLSL-style tessellation, uniformly increases resolution on all triangles within Modifier bounds */
+	UniformRings UMETA(DisplayName = "Uniform Rings (GLSL-style)"),
+
+	/** GLSL-style tessellation, resolution increases gradually towards Modifier center */
+	AdaptiveRings UMETA(DisplayName = "Adaptive Rings (GLSL-style)"),
+
+	/** Adaptive Red-Green tessellation, resolution increases gradually towards Modifier center */
+	AdaptiveRegular  UMETA(DisplayName = "Adaptive Regular (Red Green)"),
+};
+
+UCLASS(MinimalAPI, meta=(BlueprintSpawnableComponent, MegaMeshClassVersion = "2"))
+class UTessellateModifier : public MeshPartition::UModifierComponent
+{
+	GENERATED_BODY()
+	
+public:
+
+	UE_API UTessellateModifier();
+
+	// UObject
+	UE_API virtual void Serialize(FArchive& Ar) override;
+protected:
+
+	// MeshPartition::UModifierComponent
+	UE_API virtual TArray<FBox> ComputeBounds() const override;
+	UE_API virtual TSharedPtr<const MeshPartition::IModifierBackgroundOp> CreateBackgroundOp(const MeshPartition::EBuildType InBuildType) const override;
+	UE_API virtual FGuid GetCodeVersionKey() const override;
+	virtual bool IsContiguous() const override { return true; }
+	UE_API virtual void DrawVisualization(const FSceneView* View, FPrimitiveDrawInterface* PDI) const override;
+	UE_API virtual void GatherDependencies(MeshPartition::IDependencyInterface& Dependencies) const override;
+
+	UE_API FBox ComputeLocalBounds() const;
+
+	/** Region affected by the tessellation */
+	UPROPERTY(EditAnywhere, Category = Tessellation)
+	FVector3d Extents = FVector3d(1000, 1000, 1000);
+
+	/** Type of tessellation to perform */
+	UPROPERTY(EditAnywhere, Category = Tessellation)
+	MeshPartition::ETessellateModifierMethod TessellationMethod = MeshPartition::ETessellateModifierMethod::AdaptiveRegular;
+
+	/** If true, the Level value is ignored and the target edge length is used to compute refinement level for each triangle */
+	UPROPERTY(EditAnywhere, Category = Resolution)
+	bool bUseTargetEdgeLength = false;
+
+	/** Fixed number of tessellation passes to perform */
+	UPROPERTY(EditAnywhere, Category = Resolution, meta = (EditCondition = "!bUseTargetEdgeLength", UIMin = 0, UIMax = 10, ClampMin = 0))
+	int32 Level = 1;
+
+	/** Perform a number of tessellation passes needed to hit the target edge length */
+	UPROPERTY(EditAnywhere, Category = Resolution, meta = (EditCondition = "bUseTargetEdgeLength", UIMin = 1.0, UIMax = 10000.0, ClampMin = 0.01, ClampMax = 1000000.0))
+	float TargetEdgeLength = 100.0;
+
+	/** When using the target edge length, limit the tessellation level per triangle to this value to avoid excessive computation */
+	UPROPERTY(EditAnywhere, Category = Resolution, meta = (EditCondition = "bUseTargetEdgeLength", UIMin = 0, UIMax = 10, ClampMin = 0))
+	int32 MaxLevel = 4;
+
+	/** Whether to visualize the modifier bounds */
+	UPROPERTY(EditAnywhere, Category = Tessellation)
+	bool bDrawLocalBounds = true;
+
+};
+} // namespace UE::MeshPartition
+
+#undef UE_API
diff --git a/Engine/Plugins/Experimental/MeshPartition/Source/MeshPartitionEditorUI/Private/UI/MeshPartitionLayerNameWidget.cpp b/Engine/Plugins/Experimental/MeshPartition/Source/MeshPartitionEditorUI/Private/UI/MeshPartitionLayerNameWidget.cpp
--- a/Engine/Plugins/Experimental/MeshPartition/Source/MeshPartitionEditorUI/Private/UI/MeshPartitionLayerNameWidget.cpp
+++ b/Engine/Plugins/Experimental/MeshPartition/Source/MeshPartitionEditorUI/Private/UI/MeshPartitionLayerNameWidget.cpp
@@ -48,6 +48,7 @@ namespace MegaMeshNameWidgetLocals
 		{MeshPartition::UProjectMeshLayersModifier::StaticClass(), TEXT("ModifierSculpt")},
 		{MeshPartition::URemeshModifier::StaticClass(), TEXT("ModifierRemesh")},
 		{MeshPartition::USplineModifier::StaticClass(), TEXT("ModifierPath")},
+		//{MeshPartition::UTessellateModifier::StaticClass(), TEXT("Icons.Star")},
 		{MeshPartition::UTexturePatchModifier::StaticClass(), TEXT("ModifierTexturePatch")},
 	};
 
diff --git a/Engine/Plugins/Experimental/MeshPartition/Source/MeshPartitionModelingToolset/Private/MeshPartitionPlaceModifierTool.cpp b/Engine/Plugins/Experimental/MeshPartition/Source/MeshPartitionModelingToolset/Private/MeshPartitionPlaceModifierTool.cpp
--- a/Engine/Plugins/Experimental/MeshPartition/Source/MeshPartitionModelingToolset/Private/MeshPartitionPlaceModifierTool.cpp
+++ b/Engine/Plugins/Experimental/MeshPartition/Source/MeshPartitionModelingToolset/Private/MeshPartitionPlaceModifierTool.cpp
@@ -27,6 +27,7 @@
 #include "Modifiers/MeshPartitionNoiseModifier.h"
 #include "Modifiers/MeshPartitionBooleanModifier.h"
 #include "Modifiers/MeshPartitionLatticeModifier.h"
+#include "Modifiers/MeshPartitionTessellateModifier.h"
 #include "Modifiers/MeshPartitionSplineRemeshModifier.h"
 #include "Modifiers/MeshPartitionWeightUtilityModifier.h"
 
@@ -236,6 +237,8 @@ TObjectPtr<UClass> UPlaceModifierTool::GetModifierToCreateClass()
 		return MeshPartition::UBooleanModifier::StaticClass();
 	case MeshPartition::EModifierClassType::Lattice:
 		return MeshPartition::ULatticeModifier::StaticClass();
+	case MeshPartition::EModifierClassType::Tessellate:
+		return MeshPartition::UTessellateModifier::StaticClass();
 	case MeshPartition::EModifierClassType::SplineRemesh:
 		return MeshPartition::USplineRemeshModifier::StaticClass();
 	case MeshPartition::EModifierClassType::WeightUtility:
@@ -312,6 +315,8 @@ void UPlaceModifierTool::SetupNewModifier(AActor* ContainingActor, MeshPartition
 		break;
 	case MeshPartition::EModifierClassType::Lattice:
 		break;
+	case MeshPartition::EModifierClassType::Tessellate:
+		break;
 	case MeshPartition::EModifierClassType::SplineRemesh:
 	{
 		FCreateComponentParams Params;
diff --git a/Engine/Plugins/Experimental/MeshPartition/Source/MeshPartitionModelingToolset/Public/MeshPartitionPlaceModifierTool.h b/Engine/Plugins/Experimental/MeshPartition/Source/MeshPartitionModelingToolset/Public/MeshPartitionPlaceModifierTool.h
--- a/Engine/Plugins/Experimental/MeshPartition/Source/MeshPartitionModelingToolset/Public/MeshPartitionPlaceModifierTool.h
+++ b/Engine/Plugins/Experimental/MeshPartition/Source/MeshPartitionModelingToolset/Public/MeshPartitionPlaceModifierTool.h
@@ -28,6 +28,7 @@ enum class EModifierClassType : uint8
 	Noise,
 	Boolean,
 	Lattice,
+	Tessellate,
 	SplineRemesh,
 	WeightUtility,
 
diff --git a/Engine/Plugins/Experimental/MeshTerrainMode/Source/MeshTerrainMode/Private/MeshTerrainMode.cpp b/Engine/Plugins/Experimental/MeshTerrainMode/Source/MeshTerrainMode/Private/MeshTerrainMode.cpp
--- a/Engine/Plugins/Experimental/MeshTerrainMode/Source/MeshTerrainMode/Private/MeshTerrainMode.cpp
+++ b/Engine/Plugins/Experimental/MeshTerrainMode/Source/MeshTerrainMode/Private/MeshTerrainMode.cpp
@@ -1052,6 +1052,7 @@ void UMeshTerrainMode::Enter()
 	RegisterPlaceModifierTool(ToolManagerCommands.BeginAddNoiseModifierTool, TEXT("BeginAddNoiseModifierTool"), static_cast<int>(UE::MeshPartition::EModifierClassType::Noise));
 	RegisterPlaceModifierTool(ToolManagerCommands.BeginAddBooleanModifierTool, TEXT("BeginAddBooleanModifierTool"), static_cast<int>(UE::MeshPartition::EModifierClassType::Boolean));
 	RegisterPlaceModifierTool(ToolManagerCommands.BeginAddLatticeModifierTool, TEXT("BeginAddLatticeModifierTool"), static_cast<int>(UE::MeshPartition::EModifierClassType::Lattice));
+	RegisterPlaceModifierTool(ToolManagerCommands.BeginAddTessellateModifierTool, TEXT("BeginAddTessellateModifierTool"), static_cast<int>(UE::MeshPartition::EModifierClassType::Tessellate));
 	RegisterPlaceModifierTool(ToolManagerCommands.BeginAddSplineRemeshModifierTool, TEXT("BeginAddSplineRemeshModifierTool"), static_cast<int>(UE::MeshPartition::EModifierClassType::SplineRemesh));
 
 	// register extensions
diff --git a/Engine/Plugins/Experimental/MeshTerrainMode/Source/MeshTerrainMode/Private/MeshTerrainModeManagerActions.cpp b/Engine/Plugins/Experimental/MeshTerrainMode/Source/MeshTerrainMode/Private/MeshTerrainModeManagerActions.cpp
--- a/Engine/Plugins/Experimental/MeshTerrainMode/Source/MeshTerrainMode/Private/MeshTerrainModeManagerActions.cpp
+++ b/Engine/Plugins/Experimental/MeshTerrainMode/Source/MeshTerrainMode/Private/MeshTerrainModeManagerActions.cpp
@@ -328,6 +328,7 @@ void FMeshTerrainModeManagerCommands::RegisterCommands()
 	UI_COMMAND(BeginAddNoiseModifierTool, "Noise", "Add Noise Modifier", EUserInterfaceActionType::ToggleButton, FInputChord());
 	UI_COMMAND(BeginAddBooleanModifierTool, "Boolean", "Add Boolean Modifier", EUserInterfaceActionType::ToggleButton, FInputChord());
 	UI_COMMAND(BeginAddLatticeModifierTool, "Lattice", "Add Lattice Modifier", EUserInterfaceActionType::ToggleButton, FInputChord());
+	UI_COMMAND(BeginAddTessellateModifierTool, "Tessellate", "Add Tessellate Modifier", EUserInterfaceActionType::ToggleButton, FInputChord());
 	UI_COMMAND(BeginAddSplineRemeshModifierTool, "Spline Remesh", "Add Spline Remesh Modifier", EUserInterfaceActionType::ToggleButton, FInputChord());
 
 	UI_COMMAND(EnterShapesSubmode, "Shapes", "Shapes", EUserInterfaceActionType::ToggleButton, FInputChord());
diff --git a/Engine/Plugins/Experimental/MeshTerrainMode/Source/MeshTerrainMode/Private/MeshTerrainModeManagerActions.h b/Engine/Plugins/Experimental/MeshTerrainMode/Source/MeshTerrainMode/Private/MeshTerrainModeManagerActions.h
--- a/Engine/Plugins/Experimental/MeshTerrainMode/Source/MeshTerrainMode/Private/MeshTerrainModeManagerActions.h
+++ b/Engine/Plugins/Experimental/MeshTerrainMode/Source/MeshTerrainMode/Private/MeshTerrainModeManagerActions.h
@@ -270,6 +270,7 @@ public:
 	TSharedPtr<FUICommandInfo> BeginAddNoiseModifierTool;
 	TSharedPtr<FUICommandInfo> BeginAddBooleanModifierTool;
 	TSharedPtr<FUICommandInfo> BeginAddLatticeModifierTool;
+	TSharedPtr<FUICommandInfo> BeginAddTessellateModifierTool;
 	TSharedPtr<FUICommandInfo> BeginAddSplineRemeshModifierTool;
 
 	// Submodes
diff --git a/Engine/Plugins/Experimental/MeshTerrainMode/Source/MeshTerrainMode/Private/MeshTerrainModeStyle.cpp b/Engine/Plugins/Experimental/MeshTerrainMode/Source/MeshTerrainMode/Private/MeshTerrainModeStyle.cpp
--- a/Engine/Plugins/Experimental/MeshTerrainMode/Source/MeshTerrainMode/Private/MeshTerrainModeStyle.cpp
+++ b/Engine/Plugins/Experimental/MeshTerrainMode/Source/MeshTerrainMode/Private/MeshTerrainModeStyle.cpp
@@ -527,6 +527,8 @@ void FMeshTerrainModeStyle::Initialize()
 		StyleSet->Set("MeshTerrainModeManagerCommands.BeginAddBooleanModifierTool.Small",			new IMAGE_BRUSH_SVG("Icons/ModifierBoolean_16", Icon20x20));
 		StyleSet->Set("MeshTerrainModeManagerCommands.BeginAddLatticeModifierTool",				new IMAGE_BRUSH_SVG("Icons/ModifierLattice_16", Icon40x40));
 		StyleSet->Set("MeshTerrainModeManagerCommands.BeginAddLatticeModifierTool.Small",			new IMAGE_BRUSH_SVG("Icons/ModifierLattice_16", Icon20x20));
+		StyleSet->Set("MeshTerrainModeManagerCommands.BeginAddTessellateModifierTool",				new IMAGE_BRUSH_SVG("Icons/ModifierTessellation_16", Icon40x40));
+		StyleSet->Set("MeshTerrainModeManagerCommands.BeginAddTessellateModifierTool.Small",		new IMAGE_BRUSH_SVG("Icons/ModifierTessellation_16", Icon20x20));
 		StyleSet->Set("MeshTerrainModeManagerCommands.BeginAddSplineRemeshModifierTool",			new IMAGE_BRUSH_SVG("Icons/ModifierPath_16", Icon40x40));
 		StyleSet->Set("MeshTerrainModeManagerCommands.BeginAddSplineRemeshModifierTool.Small",		new IMAGE_BRUSH_SVG("Icons/ModifierPath_16", Icon20x20));
 
